Rename CustomLaunchUI to ProjectLauncher
- rename 'marshal' concept to 'tree builder' to improve clarity

diff --git a/Engine/Plugins/Developer/ProjectLauncher/Source/CommonLaunchExtensions/Private/CommonLaunchExtensionsModule.cpp b/Engine/Plugins/Developer/ProjectLauncher/Source/CommonLaunchExtensions/Private/CommonLaunchExtensionsModule.cpp
new file mode 100644
--- /dev/null
+++ b/Engine/Plugins/Developer/ProjectLauncher/Source/CommonLaunchExtensions/Private/CommonLaunchExtensionsModule.cpp
@@ -0,0 +1,42 @@
+// Copyright Epic Games, Inc. All Rights Reserved.
+
+#include "Modules/ModuleInterface.h"
+#include "Templates/SharedPointer.h"
+#include "ProjectLauncherModule.h"
+#include "Globals/GlobalsLaunchExtension.h"
+#include "Insights/InsightsLaunchExtension.h"
+
+class FCommonLaunchExtensionsModule : public IModuleInterface
+{
+public:
+
+	virtual void StartupModule() override
+	{
+		Globals = MakeShared<FGlobalsLaunchExtension>();
+		Insights = MakeShared<FInsightsLaunchExtension>();
+
+		IProjectLauncherModule::Get().RegisterExtension(Globals.ToSharedRef());
+		IProjectLauncherModule::Get().RegisterExtension(Insights.ToSharedRef());
+	}
+
+	virtual void ShutdownModule() override
+	{
+		if (IProjectLauncherModule* ProjectLauncher = IProjectLauncherModule::TryGet())
+		{
+			ProjectLauncher->UnregisterExtension(Globals.ToSharedRef());
+			ProjectLauncher->UnregisterExtension(Insights.ToSharedRef());
+		}
+
+		Globals.Reset();
+		Insights.Reset();
+	}
+
+private:
+	TSharedPtr<FGlobalsLaunchExtension> Globals;
+	TSharedPtr<FInsightsLaunchExtension> Insights;
+
+};
+
+
+IMPLEMENT_MODULE(FCommonLaunchExtensionsModule, CommonLaunchExtensions);
+
diff --git a/Engine/Plugins/Developer/ProjectLauncher/Source/ProjectLauncher/Private/Model/ProjectLauncherModel.cpp b/Engine/Plugins/Developer/ProjectLauncher/Source/ProjectLauncher/Private/Model/ProjectLauncherModel.cpp
new file mode 100644
--- /dev/null
+++ b/Engine/Plugins/Developer/ProjectLauncher/Source/ProjectLauncher/Private/Model/ProjectLauncherModel.cpp
@@ -0,0 +1,823 @@
+// Copyright Epic Games, Inc. All Rights Reserved.
+
+#include "Model/ProjectLauncherModel.h"
+#include "Misc/Paths.h"
+#include "Misc/ConfigCacheIni.h"
+#include "Misc/ConfigContext.h"
+#include "Misc/App.h"
+#include "PlatformInfo.h"
+#include "GameProjectHelper.h"
+#include "ITargetDeviceProxy.h"
+#include "ITargetDeviceProxyManager.h"
+#include "ITargetDeviceServicesModule.h"
+#include "Serialization/BufferArchive.h"
+#include "Serialization/MemoryReader.h"
+#include "Modules/ModuleManager.h"
+#include "Settings/ProjectPackagingSettings.h"
+#include "Experimental/ZenServerInterface.h"
+
+#define LOCTEXT_NAMESPACE "SCustomLaunchProfileSelector"
+
+namespace ProjectLauncher
+{
+	// this option filters the platforms & targets based on each other. 
+	// it makes it easier to choose the desired build target but the logic under the hood is more complex
+	bool bUseFriendlyBuildTargetSelection = true;
+
+
+	FModel::FModel(const TSharedRef<ITargetDeviceProxyManager>& InDeviceProxyManager, const TSharedRef<ILauncher>& InLauncher, const TSharedRef<ILauncherProfileManager>& InProfileManager)
+		: DeviceProxyManager(InDeviceProxyManager)
+		, Launcher(InLauncher)
+		, ProfileManager(InProfileManager)
+	{
+		// register callbacks
+		ProfileManager->OnPostProcessLaunchCommandLine().AddRaw( this, &FModel::OnModifyLaunchCommandLine );
+
+		ProfileManager->OnProfileAdded().AddRaw(this, &FModel::HandleProfileManagerProfileAdded);
+		ProfileManager->OnProfileRemoved().AddRaw(this, &FModel::HandleProfileManagerProfileRemoved);
+
+		DeviceProxyManager->OnProxyAdded().AddRaw(this, &FModel::HandleDeviceProxyAdded);
+		DeviceProxyManager->OnProxyRemoved().AddRaw(this, &FModel::HandleDeviceProxyRemoved);
+
+
+		// ensure there's a project when we're in the editor
+		// this means we don't need to display the global project selector in the editor
+		if (GIsEditor && FPaths::IsProjectFilePathSet())
+		{
+			ProfileManager->SetProjectPath(FPaths::GetProjectFilePath());
+		}
+
+
+		// prepare profiles
+		BasicLaunchProfile = CreateBasicLaunchProfile();
+		bHasSetBasicLaunchProfilePlatform = (BasicLaunchProfile->GetCookedPlatforms().Num() > 0);
+
+		AllProfiles = ProfileManager->GetAllProfiles();
+		AllProfiles.Add(BasicLaunchProfile);
+		SortProfiles();
+
+		DefaultBasicLaunchProfile = CreateBasicLaunchProfile();
+		DefaultCustomLaunchProfile = CreateCustomProfile(TEXT("DefaultCustomProfile"));
+
+
+		// use a custom ini file so that it can be shared between UnrealFrontend and the editor
+		ConfigFileName = FPaths::EngineSavedDir() / TEXT("Config") / TEXT("ProjectLauncher") / TEXT("UserSettings.ini");
+		FConfigContext::ReadSingleIntoGConfig().Load(*ConfigFileName);
+
+		LoadConfig();
+	}
+
+
+
+	FModel::~FModel()
+	{
+		ProfileManager->OnPostProcessLaunchCommandLine().RemoveAll( this );
+
+		ProfileManager->OnProfileAdded().RemoveAll(this);
+		ProfileManager->OnProfileRemoved().RemoveAll(this);
+
+		DeviceProxyManager->OnProxyAdded().RemoveAll(this);
+		DeviceProxyManager->OnProxyRemoved().RemoveAll(this);
+
+		SaveConfig();
+	}
+
+
+	void FModel::SelectProfile(const ILauncherProfilePtr& NewProfile)
+	{
+		if (SelectedProfile != NewProfile)
+		{
+			ILauncherProfilePtr PreviousProfile = SelectedProfile;
+			SelectedProfile = NewProfile;
+
+			ProfileSelectedDelegate.Broadcast(SelectedProfile, PreviousProfile);
+		}
+	}
+
+
+	EProfileType FModel::GetProfileType(const ILauncherProfileRef& Profile) const
+	{
+		if (IsAdvancedProfile(Profile))
+		{
+			return EProfileType::Advanced;
+		}
+		else if (Profile == BasicLaunchProfile)
+		{
+			return EProfileType::Basic;
+		}
+		else
+		{
+			return EProfileType::Custom;
+		}
+	}
+
+	bool FModel::IsAdvancedProfile(const ILauncherProfileRef& Profile) const
+	{
+		return Profile->GetCookedPlatforms().Num() > 1 ||
+				Profile->GetCookedCultures().Num() > 1 ||
+				Profile->IsCreatingDLC() ||
+				Profile->IsCreatingReleaseVersion() ||
+				Profile->IsGeneratingPatch() ||				
+				Profile->GetPackagingMode() == ELauncherProfilePackagingModes::SharedRepository ||
+				Profile->GetDeploymentMode() == ELauncherProfileDeploymentModes::CopyRepository ||
+				Profile->GetLaunchMode() == ELauncherProfileLaunchModes::CustomRoles
+			;
+	}
+
+
+
+
+	const TCHAR* FModel::GetConfigSection() const
+	{
+		return TEXT("ProjectLauncher"); 
+	}
+
+	const FString& FModel::GetConfigIni() const
+	{
+		return ConfigFileName;
+	}
+
+	void FModel::LoadConfig()
+	{
+		if (!GIsEditor)
+		{
+			// restore the previous project selection
+			FString ProjectPath;
+	
+			if (FPaths::IsProjectFilePathSet())
+			{
+				ProjectPath = FPaths::GetProjectFilePath();
+			}
+			else if (FGameProjectHelper::IsGameAvailable(FApp::GetProjectName()))
+			{
+				ProjectPath = FPaths::RootDir() / FApp::GetProjectName() / FApp::GetProjectName() + TEXT(".uproject");
+			}
+			else if (GConfig != nullptr)
+			{
+				GConfig->GetString(GetConfigSection(), TEXT("SelectedProjectPath"), ProjectPath, GetConfigIni());
+			}
+	
+			ProfileManager->SetProjectPath(ProjectPath);
+		}
+	}
+
+
+
+	void FModel::SaveConfig()
+	{
+		if (!GIsEditor && GConfig != nullptr && !FPaths::IsProjectFilePathSet() && !FGameProjectHelper::IsGameAvailable(FApp::GetProjectName()))
+		{
+			FString ProjectPath = ProfileManager->GetProjectPath();
+			GConfig->SetString(GetConfigSection(), TEXT("SelectedProjectPath"), *ProjectPath, GetConfigIni());
+		}
+	}
+
+
+
+	const PlatformInfo::FTargetPlatformInfo* FModel::GetPlatformInfo(const ILauncherProfilePtr& Profile)
+	{
+		if (Profile.IsValid() && Profile->GetCookedPlatforms().Num() > 0)
+		{
+			FString SelectedPlatform = Profile->GetCookedPlatforms()[0];
+			return PlatformInfo::FindPlatformInfo(FName(SelectedPlatform));
+		}
+
+		return nullptr;
+	}
+
+	const PlatformInfo::FTargetPlatformInfo* FModel::GetPlatformInfo(FName PlatformName, const FTargetInfo& BuildTargetInfo)
+	{
+		const PlatformInfo::FTargetPlatformInfo* PlatformInfo = PlatformInfo::FindPlatformInfo(PlatformName);
+		if (PlatformInfo == nullptr)
+		{
+			return nullptr;
+		}
+
+		// see if we found the platform immediately
+		if (BuildTargetInfo.Name.IsEmpty() || PlatformInfo->PlatformType == BuildTargetInfo.Type)
+		{
+			return PlatformInfo;
+		}
+
+		// try to find a matching flavor for the given platform & build target
+		if (PlatformInfo->VanillaInfo->PlatformType == BuildTargetInfo.Type && PlatformInfo->VanillaInfo->PlatformFlavor == PlatformInfo->PlatformFlavor)
+		{
+			return PlatformInfo->VanillaInfo;
+		}
+		for (const PlatformInfo::FTargetPlatformInfo* PlatformInfoFlavor : PlatformInfo->VanillaInfo->Flavors)
+		{
+			if (PlatformInfoFlavor->PlatformType == BuildTargetInfo.Type && PlatformInfoFlavor->PlatformFlavor == PlatformInfo->PlatformFlavor)
+			{
+				return PlatformInfoFlavor;
+			}
+		}
+
+		return nullptr;
+	}
+
+	bool FModel::IsHostPlatform(const ILauncherProfilePtr& Profile)
+	{
+		if (Profile.IsValid() && Profile->GetCookedPlatforms().Num() > 0)
+		{
+			FString SelectedPlatform = Profile->GetCookedPlatforms()[0];
+			return IsHostPlatform(FName(SelectedPlatform));
+		}
+
+		return false;
+	}
+
+	bool FModel::IsHostPlatform(FName PlatformName)
+	{
+		const PlatformInfo::FTargetPlatformInfo* PlatformInfo = PlatformInfo::FindPlatformInfo(PlatformName);
+
+		if (PlatformInfo != nullptr && PlatformInfo->IniPlatformName == FPlatformProperties::IniPlatformName())
+		{
+			return true;
+		}
+
+		return false;
+	}
+
+
+
+	FTargetInfo FModel::GetBuildTargetInfo( const FString& BuildTargetName, const FString& ProjectPath )
+	{
+		FTargetInfo Result;
+		if (!BuildTargetName.IsEmpty())
+		{
+			const TArray<FTargetInfo>& BuildTargets = FDesktopPlatformModule::Get()->GetTargetsForProject(ProjectPath);
+			for (const FTargetInfo& BuildTarget : BuildTargets)
+			{
+				if (BuildTarget.Name == BuildTargetName)
+				{
+					Result = BuildTarget;
+					break;
+				}
+			}
+		}
+
+		return MoveTemp(Result);
+	}
+
+	FTargetInfo FModel::GetBuildTargetInfo( const ILauncherProfileRef& Profile )
+	{
+		const FString& BuildTargetName = Profile->GetBuildTarget();
+		const FString& ProjectPath = Profile->GetProjectPath();
+		return GetBuildTargetInfo(BuildTargetName, ProjectPath);
+	}
+
+	FString FModel::GetVanillaPlatformName( const FString& PlatformName )
+	{
+		const PlatformInfo::FTargetPlatformInfo* PlatformInfo = PlatformInfo::FindPlatformInfo(*PlatformName);
+		if (PlatformInfo != nullptr)
+		{
+			return PlatformInfo->VanillaInfo->Name.ToString();
+		}
+
+		return PlatformName;
+	}
+
+	FString FModel::GetBuildTargetPlatformName( const FString& PlatformName, const FTargetInfo& BuildTargetInfo )
+	{
+		const PlatformInfo::FTargetPlatformInfo* PlatformInfo = ProjectLauncher::FModel::GetPlatformInfo(*PlatformName, BuildTargetInfo);
+		if (PlatformInfo != nullptr)
+		{
+			return PlatformInfo->Name.ToString();
+		}
+
+		return PlatformName;
+
+	}
+
+
+
+	void FModel::SortProfiles()
+	{
+		AllProfiles.Sort( [this](const ILauncherProfilePtr& ProfileA, const ILauncherProfilePtr& ProfileB )
+		{
+			if (ProfileA == BasicLaunchProfile)
+			{
+				return true;
+			}
+			else if (ProfileB == BasicLaunchProfile)
+			{
+				return false;
+			}
+			else
+			{
+				return ProfileA->GetName() < ProfileB->GetName();
+			}
+		});
+	}
+
+	void FModel::HandleProfileManagerProfileAdded(const ILauncherProfileRef& Profile)
+	{
+		AllProfiles.Add(Profile);
+	}
+
+
+
+	void FModel::HandleProfileManagerProfileRemoved(const ILauncherProfileRef& Profile)
+	{
+		AllProfiles.Remove(Profile);
+
+		if (Profile == SelectedProfile)
+		{
+			SelectProfile(BasicLaunchProfile);
+		}
+	}
+
+
+
+
+	void FModel::HandleDeviceProxyAdded(const TSharedRef<ITargetDeviceProxy>& DeviceProxy )
+	{
+		if (!bHasSetBasicLaunchProfilePlatform)
+		{
+			UpdatedCookedPlatformsFromDeployDeviceProxy(BasicLaunchProfile.ToSharedRef(), DeviceProxy);
+			bHasSetBasicLaunchProfilePlatform = true;
+		}
+	}
+
+	void FModel::HandleDeviceProxyRemoved(const TSharedRef<ITargetDeviceProxy>& DeviceProxy )
+	{
+	}
+
+
+
+
+	TSharedPtr<ITargetDeviceProxy> FModel::GetDeviceProxy(const ILauncherProfileRef& Profile)
+	{
+		ILauncherDeviceGroupPtr DeployedDeviceGroup = Profile->GetDeployedDeviceGroup();
+		if (DeployedDeviceGroup != nullptr && DeployedDeviceGroup->GetDeviceIDs().Num() > 0)
+		{
+			const FString& DeviceID = DeployedDeviceGroup->GetDeviceIDs()[0];
+
+			ITargetDeviceServicesModule& TargetDeviceServicesModule = FModuleManager::LoadModuleChecked<ITargetDeviceServicesModule>("TargetDeviceServices");
+
+			return TargetDeviceServicesModule.GetDeviceProxyManager()->FindProxyDeviceForTargetDevice(DeviceID);
+		}
+
+		return nullptr;
+	}
+
+
+	ILauncherProfileRef FModel::CreateCustomProfile( const TCHAR* Name )
+	{
+		//create the profile
+		ILauncherProfileRef Profile = ProfileManager->CreateUnsavedProfile(Name);
+
+		// set defaults
+		SetProfileContentScheme(EContentScheme::ZenStreaming, Profile );
+		Profile->SetBuildConfiguration(EBuildConfiguration::Development);
+		Profile->SetLaunchMode(ELauncherProfileLaunchModes::DefaultRole);
+		Profile->SetBuildMode(ELauncherProfileBuildModes::Auto);
+		Profile->SetBuildUAT(false);
+		// @fixme: set all defaults here 
+		// ...
+
+
+		// make sure there is a device & deploy group
+		ILauncherDeviceGroupRef DeployDeviceGroup = ProfileManager->AddNewDeviceGroup();
+		Profile->SetDeployedDeviceGroup(DeployDeviceGroup);
+
+		TArray<TSharedPtr<ITargetDeviceProxy>> DeviceProxies;
+		DeviceProxyManager->GetProxies(NAME_None, true, DeviceProxies);
+		if (DeviceProxies.Num() > 0 && DeviceProxies[0].IsValid())
+		{
+			UpdatedCookedPlatformsFromDeployDeviceProxy(Profile, DeviceProxies[0].ToSharedRef());
+		}
+
+		return Profile;
+	}
+
+	ILauncherProfileRef FModel::CreateBasicLaunchProfile()
+	{
+		const FText BasicLaunchProfileName = LOCTEXT("BasicLaunchProfileName", "Basic Launch");
+		const FText BasicLaunchProfileDescription = LOCTEXT("BasicLaunchProfileDescription", "Use this profile to launch on a device with the recommended defaults");
+
+		ILauncherProfileRef Profile = CreateCustomProfile(*BasicLaunchProfileName.ToString());
+		Profile->SetDescription(*BasicLaunchProfileDescription.ToString());
+		Profile->SetProjectSpecified(false);
+		Profile->SetBuildTargetSpecified(false);
+
+		return Profile;
+	}
+
+
+
+	void FModel::UpdatedCookedPlatformsFromDeployDeviceProxy(const ILauncherProfileRef& Profile, TSharedPtr<ITargetDeviceProxy> DeviceProxy)
+	{
+		if (DeviceProxy.IsValid())
+		{
+			Profile->GetDeployedDeviceGroup()->RemoveAllDevices();
+			Profile->GetDeployedDeviceGroup()->AddDevice(DeviceProxy->GetTargetDeviceId(NAME_None));
+		}
+		else if (Profile->GetDeployedDeviceGroup()->GetNumDevices() > 0)
+		{
+			DeviceProxy = GetDeviceProxy(Profile);
+		}
+		
+		if (DeviceProxy.IsValid())
+		{
+			Profile->ClearCookedPlatforms();
+		
+			FTargetInfo BuildTargetInfo = GetBuildTargetInfo(Profile);
+			FString PlatformName = DeviceProxy->GetTargetPlatformName(NAME_None);
+			const PlatformInfo::FTargetPlatformInfo* PlatformInfo = GetPlatformInfo(*PlatformName, BuildTargetInfo);
+
+			if (PlatformInfo != nullptr)
+			{
+				Profile->AddCookedPlatform(PlatformInfo->Name.ToString());
+			}
+			else
+			{
+				Profile->AddCookedPlatform(DeviceProxy->GetTargetPlatformName(NAME_None));
+			}
+
+			// @todo: is default deploy platform necessary now?
+		}
+	}
+
+	void FModel::UpdateCookedPlatformsFromBuildTarget(const ILauncherProfileRef& Profile)
+	{
+		FTargetInfo BuildTargetInfo = GetBuildTargetInfo(Profile);
+
+		TArray<FString> Platforms = Profile->GetCookedPlatforms();
+		Profile->ClearCookedPlatforms();
+		for (const FString& Platform : Platforms)
+		{
+			const PlatformInfo::FTargetPlatformInfo* PlatformInfo = GetPlatformInfo(*Platform, BuildTargetInfo);
+			if (PlatformInfo != nullptr)
+			{
+				Profile->AddCookedPlatform(PlatformInfo->Name.ToString());
+			}
+			else
+			{
+				Profile->AddCookedPlatform(Platform);
+			}
+		}
+	}
+
+
+	ILauncherProfilePtr FModel::CloneCustomProfile(const ILauncherProfileRef& Profile)
+	{
+		ILauncherProfilePtr NewProfile;
+
+		FBufferArchive Writer;
+		if (Profile->Serialize(Writer))
+		{
+			NewProfile = CreateCustomProfile(TEXT("Cloned"));
+			if (NewProfile.IsValid())
+			{
+				FMemoryReader Reader(Writer);
+				NewProfile->Serialize(Reader);
+				NewProfile->AssignId(true); // need to give it a new Id - don't use the cloned one!
+
+				// @todo: add these to FLauncherProfile's Serialize() function & update the internal version. Waiting for all Zen properties to be decided (zen workspace etc)
+				NewProfile->SetUseZenStore(Profile->IsUsingZenStore());
+				// ...
+
+				ILauncherDeviceGroupRef DeployDeviceGroup = ProfileManager->AddNewDeviceGroup();
+				NewProfile->SetDeployedDeviceGroup(DeployDeviceGroup);
+				for (const FString& DeviceID : Profile->GetDeployedDeviceGroup()->GetDeviceIDs())
+				{
+					DeployDeviceGroup->AddDevice(DeviceID);
+				}
+			}
+		}
+
+		return NewProfile;
+	}
+
+			
+
+	EContentScheme FModel::DetermineProfileContentScheme(const ILauncherProfileRef& Profile) const
+	{
+		if (Profile->GetCookMode() == ELauncherProfileCookModes::OnTheFly)
+		{
+			return EContentScheme::CookOnTheFly;
+		}
+		else if (Profile->GetPackagingMode() != ELauncherProfilePackagingModes::DoNotPackage)
+		{
+			return EContentScheme::DevelopmentPackage;
+		}
+		else if (Profile->IsPackingWithUnrealPak())
+		{
+			return EContentScheme::PakFiles;
+		}
+		else if (Profile->IsUsingZenPakStreaming())
+		{
+			return EContentScheme::ZenPakStreaming;
+		}
+		else if (Profile->IsUsingZenStore())
+		{
+			return EContentScheme::ZenStreaming;
+		}
+		else
+		{
+			return EContentScheme::LooseFiles;
+		}
+	}
+
+
+
+	void FModel::SetProfileContentScheme(EContentScheme ContentScheme, const ILauncherProfileRef& Profile, bool bWantToCook, ELauncherProfileDeploymentModes::Type DefaultDeploymentMode)
+	{
+		bool bPakFiles = (ContentScheme == EContentScheme::PakFiles || ContentScheme == EContentScheme::DevelopmentPackage);
+		bool bUseZen = (ContentScheme != EContentScheme::LooseFiles && ContentScheme != EContentScheme::CookOnTheFly);
+		bool bCOTF = (ContentScheme == EContentScheme::CookOnTheFly);
+		bool bPackage = (ContentScheme == EContentScheme::DevelopmentPackage);
+		bool bZenPakStreaming = (ContentScheme == EContentScheme::ZenPakStreaming);
+		// FIXME: may need to turn off Zen/Pak in UAT somehow if it's set in project defaults?
+
+		Profile->SetUseZenPakStreaming(bZenPakStreaming);
+		Profile->SetDeployWithUnrealPak(bPakFiles);
+		Profile->SetUseZenStore(bUseZen);
+
+		ELauncherProfileDeploymentModes::Type DeploymentMode = DefaultDeploymentMode;
+		if (bCOTF)
+		{
+			Profile->SetCookMode(ELauncherProfileCookModes::OnTheFly);
+			DeploymentMode = ELauncherProfileDeploymentModes::FileServer;
+		}
+		else if (bZenPakStreaming || !bWantToCook)
+		{
+			Profile->SetCookMode(ELauncherProfileCookModes::DoNotCook);
+		}
+		else
+		{
+			Profile->SetCookMode(ELauncherProfileCookModes::ByTheBook);
+		}
+
+		if (bPackage)
+		{
+			Profile->SetDeploymentMode(ELauncherProfileDeploymentModes::DoNotDeploy); // @todo: some platforms support package deployment. should update verification to check new function in ITargetPlatformControls
+			Profile->SetPackagingMode(ELauncherProfilePackagingModes::Locally);
+		}
+		else
+		{
+			Profile->SetDeploymentMode(DeploymentMode);
+			Profile->SetPackagingMode(ELauncherProfilePackagingModes::DoNotPackage);
+		}
+	}
+
+
+	void FModel::ReadProjectSettingsFromConfig( FConfigCacheIni& InConfig, FProjectSettings& OutResult )
+	{
+		// read project packaging settings
+		const TCHAR* ProjectPackagingConfigSection = TEXT("/Script/UnrealEd.ProjectPackagingSettings");
+		InConfig.GetBool(ProjectPackagingConfigSection, TEXT("bUseZenStore"), OutResult.bUseZenStore, GGameIni);
+		InConfig.GetBool(ProjectPackagingConfigSection, TEXT("bEnablePakStreaming"), OutResult.bHasAutomaticZenPakStreamingWorkspaceCreation, GGameIni);
+		
+		// read zen settings
+#if UE_WITH_ZEN
+		using namespace UE::Zen;
+		const FServiceSettings& ZenServiceSettings = GetDefaultServiceInstance().GetServiceSettings();
+		OutResult.bAllowRemoteNetworkService = ZenServiceSettings.IsAutoLaunch() && ZenServiceSettings.SettingsVariant.Get<FServiceAutoLaunchSettings>().bAllowRemoteNetworkService;
+#else
+		const TCHAR* ZenAutoLaunchSettings = TEXT("Zen.AutoLaunch");
+		InConfig.GetBool(ZenAutoLaunchSettings, TEXT("AllowRemoteNetworkService"), OutResult.bAllowRemoteNetworkService, GEngineIni);
+#endif
+	}
+
+	const FProjectSettings FModel::GetProjectSettings( const FString& InProjectPath )
+	{
+		FProjectSettings Result;
+
+		FString ProjectPath = FPaths::ConvertRelativePathToFull(InProjectPath);
+		FString ProjectName = FPaths::GetBaseFilename(ProjectPath);
+
+		FString GlobalProjectPath;
+		if (FPaths::IsProjectFilePathSet())
+		{
+			GlobalProjectPath = FPaths::ConvertRelativePathToFull(FPaths::GetProjectFilePath());
+		}
+
+		if (!GlobalProjectPath.IsEmpty() && GlobalProjectPath == ProjectPath )
+		{
+			// read the current project's ini settings
+			Result.bIsCurrentEditorProject = true;
+			ReadProjectSettingsFromConfig(*GConfig, Result);
+
+			// apply any in-memory properties @fixme: is this necessary?
+			const UProjectPackagingSettings* ProjectPackagingSettings = UProjectPackagingSettings::StaticClass()->GetDefaultObject<UProjectPackagingSettings>();
+			Result.bUseZenStore = ProjectPackagingSettings->bUseZenStore;
+
+			// note this is not cached because the user could edit the properties from within the editor
+		}
+		else if (CachedProjectSettings.Contains(ProjectName))
+		{
+			Result = CachedProjectSettings[ProjectName];
+		}
+		else if (!ProjectPath.IsEmpty())
+		{
+			// load the other project's ini files into a temporary config cache
+			FConfigCacheIni Config(EConfigCacheType::Temporary);
+			FConfigContext Context = FConfigContext::ReadIntoConfigSystem(&Config, FString());
+			Context.ProjectConfigDir = FPaths::Combine(FPaths::GetPath(ProjectPath), TEXT("Config/"));
+			Config.InitializeKnownConfigFiles(Context);
+
+			// read the ini settings
+			Result.bIsCurrentEditorProject = false;
+			ReadProjectSettingsFromConfig(Config, Result);
+
+			// cache the result
+			CachedProjectSettings.Add( ProjectName, Result);
+		}
+
+		return MoveTemp(Result);
+	}
+
+	const FProjectSettings FModel::GetProjectSettings( const ILauncherProfileRef& Profile )
+	{
+		return GetProjectSettings(*Profile->GetProjectPath());
+	}
+
+
+
+	bool FModel::IsBasicLaunchProfile( const ILauncherProfilePtr& Profile ) const
+	{
+		return (Profile == BasicLaunchProfile);
+	}
+
+
+	const ILauncherProfileRef FModel::GetDefaultBasicLaunchProfile() const
+	{
+		return DefaultBasicLaunchProfile.ToSharedRef();
+	}
+
+	const ILauncherProfileRef FModel::GetDefaultCustomLaunchProfile() const
+	{
+		return DefaultCustomLaunchProfile.ToSharedRef();
+	}
+
+
+
+	TSharedPtr<FLaunchLogMessage> FModel::AddLogMessage( const FString& InMessage, ELogVerbosity::Type InVerbosity )
+	{
+		TSharedPtr<FLaunchLogMessage> Message = MakeShared<FLaunchLogMessage>(InMessage, InVerbosity);
+		LaunchLogMessages.Add(Message);
+		return Message;
+	}
+
+	void FModel::ClearLogMessages()
+	{
+		LaunchLogMessages.Reset();
+	}
+
+
+
+	extern void ApplyExtensionVariables( const ILauncherProfileRef& InProfile, FString& InOutCommandLine, TSharedRef<FModel> InModel );
+
+	void FModel::OnModifyLaunchCommandLine( const ILauncherProfileRef& InProfile, FString& InOutCommandLine )
+	{
+		ApplyExtensionVariables(InProfile, InOutCommandLine, AsShared() );
+	}
+
+
+
+
+
+
+
+
+
+
+
+	TArray<EContentScheme> GetAllContentSchemes()
+	{
+		TArray<EContentScheme> ContentSchemes;
+		for (uint8 Index = 0; Index < (uint8)EContentScheme::MAX; Index++)
+		{
+			ContentSchemes.Add((EContentScheme)Index);
+		}
+
+		return MoveTemp(ContentSchemes);
+	}
+
+	FText GetContentSchemeDisplayName(EContentScheme ContentScheme)
+	{
+		switch (ContentScheme)
+		{
+			case EContentScheme::PakFiles:           return LOCTEXT("ContentSchemePakFiles","Pak Files"); 
+			case EContentScheme::ZenStreaming:       return LOCTEXT("ContentSchemeZenStreaming","Zen Streaming");
+			case EContentScheme::ZenPakStreaming:    return LOCTEXT("ContentSchemeZenPakStreaming","Zen Pak Streaming"); 
+			case EContentScheme::DevelopmentPackage: return LOCTEXT("ContentSchemeDevPackage","Development Package");
+			case EContentScheme::LooseFiles:         return LOCTEXT("ContentSchemeLooseFiles","Loose Files (legacy)");
+			case EContentScheme::CookOnTheFly:       return LOCTEXT("ContentSchemeCOTF","Cook On The Fly"); 
+		}
+
+		checkNoEntry();
+		return FText::GetEmpty();
+
+	}
+
+	FText GetContentSchemeToolTip(EContentScheme ContentScheme)
+	{
+		switch (ContentScheme)
+		{
+			case EContentScheme::PakFiles:           return LOCTEXT("ContentSchemeTipPakFiles","Store cooked game content in one or more large Pak Files");
+			case EContentScheme::ZenStreaming:       return LOCTEXT("ContentSchemeTipZenStreaming","Stream cooked game content from Zen Server");
+			case EContentScheme::ZenPakStreaming:    return LOCTEXT("ContentSchemeTipZenPakStreaming","Stream an existing Pak Files build via a Zen");
+			case EContentScheme::DevelopmentPackage: return LOCTEXT("ContentSchemeTipDevPackage","Package cooked game content into a single installable package file for development purposes, where available");
+			case EContentScheme::LooseFiles:         return LOCTEXT("ContentSchemeTipLooseFiles","Store cooked game assets in individual files (legacy - recommend moving to Zen Streaming. This option will not work if the project is already configured with 'Use Zen Store')");
+			case EContentScheme::CookOnTheFly:       return LOCTEXT("ContentSchemeTipCOTF","Only cook game assets when the game requires them, and send them over the network (legacy - slow)");
+		}
+
+		checkNoEntry();
+		return FText::GetEmpty();
+
+	}
+
+
+
+
+	const TCHAR* LexToString( const ProjectLauncher::EContentScheme& ContentScheme)
+	{
+		switch (ContentScheme)
+		{
+			case ProjectLauncher::EContentScheme::PakFiles:           return TEXT("PakFiles");
+			case ProjectLauncher::EContentScheme::ZenStreaming:       return TEXT("ZenStreaming");
+			case ProjectLauncher::EContentScheme::ZenPakStreaming:    return TEXT("ZenPakStreaming");
+			case ProjectLauncher::EContentScheme::DevelopmentPackage: return TEXT("DevelopmentPackage");
+			case ProjectLauncher::EContentScheme::LooseFiles:         return TEXT("LooseFiles");
+			case ProjectLauncher::EContentScheme::CookOnTheFly:       return TEXT("CookOnTheFly");
+			default: return TEXT("Unknown");
+		}
+	}
+
+
+
+	bool LexTryParseString( ProjectLauncher::EContentScheme& OutContentScheme, const TCHAR* String )
+	{
+		if (FCString::Stricmp(String, TEXT("PakFiles")) == 0) 
+		{
+			OutContentScheme = ProjectLauncher::EContentScheme::PakFiles;
+			return true;
+		}
+		else if (FCString::Stricmp(String, TEXT("ZenStreaming")) == 0) 
+		{
+			OutContentScheme = ProjectLauncher::EContentScheme::ZenStreaming;
+			return true;
+		}
+		else if (FCString::Stricmp(String, TEXT("ZenPakStreaming")) == 0) 
+		{
+			OutContentScheme = ProjectLauncher::EContentScheme::ZenPakStreaming;
+			return true;
+		}
+		else if (FCString::Stricmp(String, TEXT("DevelopmentPackage")) == 0) 
+		{
+			OutContentScheme = ProjectLauncher::EContentScheme::DevelopmentPackage;
+			return true;
+		}
+		else if (FCString::Stricmp(String, TEXT("LooseFiles")) == 0) 
+		{
+			OutContentScheme = ProjectLauncher::EContentScheme::LooseFiles;
+			return true;
+		}
+		else if (FCString::Stricmp(String, TEXT("CookOnTheFly")) == 0) 
+		{
+			OutContentScheme = ProjectLauncher::EContentScheme::CookOnTheFly;
+			return true;
+		}
+		else
+		{
+			return false;
+		}
+	}
+
+
+
+	FText GetProfileLaunchErrorMessage(ILauncherProfilePtr Profile)
+	{
+		if (!Profile.IsValid())
+		{
+			return LOCTEXT("LaunchErrNoProfileTip", "There is no profile selected");
+		}
+
+		if (Profile->IsValidForLaunch())
+		{
+			return LOCTEXT("LaunchProfileTip", "Launch this profile now");
+		}
+
+		FTextBuilder MsgTextBuilder;
+		MsgTextBuilder.AppendLine(LOCTEXT("LaunchErrValidation", "There are validation errors with this profile. Please fix them before launching:"));
+		MsgTextBuilder.Indent();
+		for (int i = 0; i < (int)ELauncherProfileValidationErrors::Count; i++)
+		{
+			ELauncherProfileValidationErrors::Type Error = (ELauncherProfileValidationErrors::Type)i;
+			if (Profile->HasValidationError(Error))
+			{
+				MsgTextBuilder.AppendLine(LexToStringLocalized(Error));
+			}
+		}
+		MsgTextBuilder.Unindent();
+
+		return MsgTextBuilder.ToText();
+	}
+}
+
+
+
+#undef LOCTEXT_NAMESPACE
diff --git a/Engine/Plugins/Developer/ProjectLauncher/Source/ProjectLauncher/Private/ProfileTree/BasicProfileTreeBuilder.cpp b/Engine/Plugins/Developer/ProjectLauncher/Source/ProjectLauncher/Private/ProfileTree/BasicProfileTreeBuilder.cpp
new file mode 100644
--- /dev/null
+++ b/Engine/Plugins/Developer/ProjectLauncher/Source/ProjectLauncher/Private/ProfileTree/BasicProfileTreeBuilder.cpp
@@ -0,0 +1,38 @@
+// Copyright Epic Games, Inc. All Rights Reserved.
+
+#include "ProfileTree/BasicProfileTreeBuilder.h"
+
+#define LOCTEXT_NAMESPACE "BasicProfileTreeBuilder"
+
+namespace ProjectLauncher
+{
+	FBasicProfileTreeBuilder::FBasicProfileTreeBuilder( const ILauncherProfileRef& InProfile, const TSharedRef<FModel>& InModel )
+		: FGenericProfileTreeBuilder( InProfile, InModel->GetDefaultBasicLaunchProfile(), InModel )
+	{
+	}
+
+
+	void FBasicProfileTreeBuilder::Construct()
+	{
+		FGenericProfileTreeBuilder::Construct();
+
+		FLaunchProfileTreeNode& GeneralSettingsHeader = TreeData->AddHeading( LOCTEXT("GeneralSettingsHeading", "General Settings") );
+		AddProjectProperty(GeneralSettingsHeader);
+		AddTargetProperty(GeneralSettingsHeader);
+		AddConfigurationProperty(GeneralSettingsHeader);
+		AddContentSchemeProperty(GeneralSettingsHeader);
+		AddTargetDeviceProperty(GeneralSettingsHeader);
+		AddCommandLineProperty(GeneralSettingsHeader);
+	}
+
+
+
+	TSharedPtr<ILaunchProfileTreeBuilder> FBasicProfileTreeBuilderFactory::TryCreateTreeBuilder( const ILauncherProfileRef& InProfile, const TSharedRef<FModel>& InModel )
+	{
+		return MakeShared<FBasicProfileTreeBuilder>(InProfile, InModel);
+	}
+
+
+}
+
+#undef LOCTEXT_NAMESPACE
diff --git a/Engine/Plugins/Developer/ProjectLauncher/Source/ProjectLauncher/Private/ProfileTree/CustomProfileTreeBuilder.cpp b/Engine/Plugins/Developer/ProjectLauncher/Source/ProjectLauncher/Private/ProfileTree/CustomProfileTreeBuilder.cpp
new file mode 100644
--- /dev/null
+++ b/Engine/Plugins/Developer/ProjectLauncher/Source/ProjectLauncher/Private/ProfileTree/CustomProfileTreeBuilder.cpp
@@ -0,0 +1,66 @@
+// Copyright Epic Games, Inc. All Rights Reserved.
+
+#include "ProfileTree/CustomProfileTreeBuilder.h"
+
+
+#define LOCTEXT_NAMESPACE "CustomProfileTreeBuilder"
+
+namespace ProjectLauncher
+{
+	FCustomProfileTreeBuilder::FCustomProfileTreeBuilder( const ILauncherProfileRef& InProfile, const TSharedRef<FModel>& InModel )
+		: FGenericProfileTreeBuilder( InProfile, InModel->GetDefaultCustomLaunchProfile(), InModel )
+	{
+	}
+
+	void FCustomProfileTreeBuilder::Construct()
+	{
+		FGenericProfileTreeBuilder::Construct();
+
+		FLaunchProfileTreeNode& GeneralSettingsHeader = TreeData->AddHeading( LOCTEXT("GeneralSettingsHeading", "General Settings") );
+		AddProjectProperty(GeneralSettingsHeader);
+		AddTargetProperty(GeneralSettingsHeader);
+		AddPlatformProperty(GeneralSettingsHeader);
+		AddConfigurationProperty(GeneralSettingsHeader);
+		AddContentSchemeProperty(GeneralSettingsHeader);
+
+		FLaunchProfileTreeNode& ContentSchemeHeader = TreeData->AddHeading( LOCTEXT("ContentSchemeHeading", "Content Scheme") );
+		AddCompressPakFilesProperty(ContentSchemeHeader);
+		AddUseIoStoreProperty(ContentSchemeHeader);
+		AddGenerateChunksProperty(ContentSchemeHeader);
+		AddImportZenSnapshotProperty(ContentSchemeHeader);
+		AddZenPakStreamingPathProperty(ContentSchemeHeader);
+
+		FLaunchProfileTreeNode& MapsAndCookingHeader = TreeData->AddHeading( LOCTEXT("CookingHeading", "Maps And Cooking") );
+		AddCookProperty(MapsAndCookingHeader);
+		AddIncrementalCookProperty(MapsAndCookingHeader);
+		AddMapsToCookProperty(MapsAndCookingHeader);
+		AddAdditionalCookerOptionsProperty(MapsAndCookingHeader);
+
+		FLaunchProfileTreeNode& BuildHeader = TreeData->AddHeading( LOCTEXT("BuildHeading", "Build") );
+		AddBuildProperty(BuildHeader);
+		AddForceBuildProperty(BuildHeader);
+		AddArchitectureProperty(BuildHeader);
+
+		FLaunchProfileTreeNode& DirectoryHeader = TreeData->AddHeading( LOCTEXT("DirectoryHeading", "Directory") );
+		AddArchiveBuildProperty(DirectoryHeader);
+		AddArchiveBuildDirectoryProperty(DirectoryHeader);
+
+		FLaunchProfileTreeNode& DeployAndRunHeader = TreeData->AddHeading( LOCTEXT("DeployAndRunHeading", "Deploy And Run") );
+		AddDeployProperty(DeployAndRunHeader);
+		AddTargetDeviceProperty(DeployAndRunHeader);
+		AddRunProperty(DeployAndRunHeader);
+		AddInitialMapProperty(DeployAndRunHeader);
+		AddCommandLineProperty(DeployAndRunHeader);
+	}
+
+
+
+	TSharedPtr<ILaunchProfileTreeBuilder> FCustomProfileTreeBuilderFactory::TryCreateTreeBuilder( const ILauncherProfileRef& InProfile, const TSharedRef<FModel>& InModel )
+	{
+		return MakeShared<FCustomProfileTreeBuilder>(InProfile, InModel);
+	}
+
+
+}
+
+#undef LOCTEXT_NAMESPACE
diff --git a/Engine/Plugins/Developer/ProjectLauncher/Source/ProjectLauncher/Private/ProfileTree/GenericProfileTreeBuilder.cpp b/Engine/Plugins/Developer/ProjectLauncher/Source/ProjectLauncher/Private/ProfileTree/GenericProfileTreeBuilder.cpp
new file mode 100644
--- /dev/null
+++ b/Engine/Plugins/Developer/ProjectLauncher/Source/ProjectLauncher/Private/ProfileTree/GenericProfileTreeBuilder.cpp
@@ -0,0 +1,1362 @@
+// Copyright Epic Games, Inc. All Rights Reserved.
+
+#include "ProfileTree/GenericProfileTreeBuilder.h"
+#include "Widgets/Shared/SCustomLaunchProjectCombo.h"
+#include "Widgets/Shared/SCustomLaunchBuildTargetCombo.h"
+#include "Widgets/Shared/SCustomLaunchPlatformCombo.h"
+#include "Widgets/Shared/SCustomLaunchContentSchemeCombo.h"
+#include "Widgets/Shared/SCustomLaunchDeviceListView.h"
+#include "Widgets/Shared/SCustomLaunchMapListView.h"
+#include "Widgets/Shared/SCustomLaunchCombo.h"
+#include "Widgets/Input/SSegmentedControl.h"
+#include "Widgets/Text/STextBlock.h"
+#include "Widgets/Layout/SBox.h"
+#include "Widgets/SBoxPanel.h"
+#include "SResizeBox.h"
+#include "SSearchableComboBox.h"
+#include "PlatformInfo.h"
+#include "IDesktopPlatform.h"
+#include "Interfaces/ITargetPlatformManagerModule.h"
+#include "Interfaces/ITargetPlatform.h"
+#include "GameProjectHelper.h"
+
+
+
+#define LOCTEXT_NAMESPACE "CustomProfileTreeBuilder"
+
+
+namespace ProjectLauncher
+{
+	FGenericProfileTreeBuilder::FGenericProfileTreeBuilder( const ILauncherProfileRef& InProfile, const ILauncherProfileRef& InDefaultProfile, const TSharedRef<FModel>& InModel )
+		: TreeData( MakeShared<FLaunchProfileTreeData>(InProfile, InModel, this) )
+		, Profile(InProfile)
+		, DefaultProfile(InDefaultProfile)
+		, Model(InModel)
+	{
+		ProfileType = Model->GetProfileType(Profile);
+
+		ForPak = [this]()
+		{
+			return (ContentScheme == EContentScheme::PakFiles || ContentScheme == EContentScheme::DevelopmentPackage);
+		};
+
+		ForZenWS = [this]()
+		{
+			return (ContentScheme == EContentScheme::ZenPakStreaming);
+		};
+
+		ForCooked = [this]()
+		{
+			return (ContentScheme != EContentScheme::ZenPakStreaming && ContentScheme != EContentScheme::CookOnTheFly);
+		};
+
+		ForEnabledCooked = [this]()
+		{
+			return (bShouldCook && ContentScheme != EContentScheme::ZenPakStreaming && ContentScheme != EContentScheme::CookOnTheFly);
+		};
+
+		ForContent = [this]()
+		{
+			return (ContentScheme != EContentScheme::ZenPakStreaming);
+		};
+
+		ForCode = [this]()
+		{
+			return GetBuild();
+		};
+
+		ForDeployment = [this]()
+		{
+			return GetDeployToDevice();
+		};
+
+		ForRun = [this]()
+		{
+			return GetIsRunning();
+		};
+
+		EmptyString = []()
+		{
+			return FString();
+		};
+
+	}
+
+	void FGenericProfileTreeBuilder::Construct()
+	{
+		const TArray<FString>& DeviceIDs = Profile->GetDeployedDeviceGroup()->GetDeviceIDs();
+
+		MapOption = Profile->GetCookedMaps().Num() > 0 ? EMapOption::Selected : EMapOption::Startup;
+		DeployDeviceOption = DeviceIDs.Num() > 0 && !DeviceIDs[0].IsEmpty() ? EDeployDeviceOption::Selected : EDeployDeviceOption::Default;
+		bShouldCook = GetCook();
+		ContentScheme = Model->DetermineProfileContentScheme(Profile);
+		CachedBuildTargetType = Model->GetBuildTargetInfo(Profile).Type;
+	
+		CacheStartupMapList();
+		CacheArchitectures();
+	}
+
+
+
+
+
+
+
+
+
+	void FGenericProfileTreeBuilder::AddProjectProperty( FLaunchProfileTreeNode& HeadingNode )
+	{
+		if (ProfileType == EProfileType::Basic)
+		{
+			// this is adding a new widget to the property tree
+			// - the first entry is the name of the property on the left-hand side
+			// - the seecond parameter is the widget itself that appears on the right-hand side
+			HeadingNode.AddWidget( LOCTEXT("ProjectLabel", "Project"), 
+				SNew(SCustomLaunchProjectCombo)
+				.OnSelectionChanged(this, &FGenericProfileTreeBuilder::SetProjectName)
+				.SelectedProject(this, &FGenericProfileTreeBuilder::GetProjectPath)
+				.HasProject(this, &FGenericProfileTreeBuilder::HasProject)
+				.CurrentProjectOption(SCustomLaunchProjectCombo::ECurrentProjectOption::Empty)
+				.Font(FCoreStyle::Get().GetFontStyle("SmallFont"))
+			);
+		}
+		else
+		{
+			HeadingNode.AddWidget( LOCTEXT("ProjectLabel", "Project"), 
+				SNew(SCustomLaunchProjectCombo)
+				.OnSelectionChanged(this, &FGenericProfileTreeBuilder::SetProjectName)
+				.SelectedProject(this, &FGenericProfileTreeBuilder::GetProjectPath)
+				.HasProject(this, &FGenericProfileTreeBuilder::HasProject)
+				.ShowAnyProjectOption(true)
+				.CurrentProjectOption(SCustomLaunchProjectCombo::ECurrentProjectOption::ActualProject)
+				.Font(FCoreStyle::Get().GetFontStyle("SmallFont"))
+			);
+		}
+	}
+
+	void FGenericProfileTreeBuilder::AddTargetProperty( FLaunchProfileTreeNode& HeadingNode )
+	{
+		if (ProfileType == EProfileType::Basic)
+		{
+			// this is also adding a new widget to the property tree, as above.
+			// in this example the new struct parameter defines several callbacks that handle the 'reset to default' functionality. there are also options for disabling & hiding.
+			// the code is implemented with this syntax to avoid aid readability without using slate's TAttribute style functionality which seemed like an overkill for our simpler needs.
+			HeadingNode.AddWidget( LOCTEXT("TargetLabel", "Target"), 
+				{
+					.IsDefault = [this]()		{ return !Profile->HasBuildTargetSpecified() || Profile->GetBuildTarget().IsEmpty(); },
+					.SetToDefault = [this]()	{ SetBuildTarget(FString()); },
+				},
+				SNew(SCustomLaunchBuildTargetCombo)
+				.OnSelectionChanged(this, &FGenericProfileTreeBuilder::SetBuildTarget)
+				.SelectedBuildTarget(this, &FGenericProfileTreeBuilder::GetBuildTarget)
+				.SelectedProject(this, &FGenericProfileTreeBuilder::GetProjectPath)
+				.Font(FCoreStyle::Get().GetFontStyle("SmallFont"))
+			);
+		}
+		else
+		{
+			HeadingNode.AddWidget( LOCTEXT("TargetLabel", "Target"), 
+				{
+					.IsDefault = [this]()		{ return !Profile->HasBuildTargetSpecified() || Profile->GetBuildTarget().IsEmpty(); },
+					.SetToDefault = [this]()	{ SetBuildTarget(FString()); },
+					.IsEnabled = [this]()		{ return Profile->HasProjectSpecified(); },
+				},
+				SNew(SCustomLaunchBuildTargetCombo)
+				.OnSelectionChanged(this, &FGenericProfileTreeBuilder::SetBuildTarget)
+				.SelectedBuildTarget(this, &FGenericProfileTreeBuilder::GetBuildTarget)
+				.SupportedTargetTypes(this, &FGenericProfileTreeBuilder::GetSupportedBuildTargetTypes)
+				.SelectedProject(this, &FGenericProfileTreeBuilder::GetProjectPath)
+				.Font(FCoreStyle::Get().GetFontStyle("SmallFont"))
+			);
+		}
+	}
+
+	void FGenericProfileTreeBuilder::AddPlatformProperty( FLaunchProfileTreeNode& HeadingNode )
+	{
+		HeadingNode.AddWidget( LOCTEXT("PlatformLabel", "Platform"), 
+			SNew(SCustomLaunchPlatformCombo)
+			.SelectedPlatforms(this, &FGenericProfileTreeBuilder::GetSelectedPlatforms)
+			.OnSelectionChanged(this, &FGenericProfileTreeBuilder::SetSelectedPlatforms)
+			.BasicPlatformsOnly(bUseFriendlyBuildTargetSelection)
+			.Font(FCoreStyle::Get().GetFontStyle("SmallFont"))
+		);
+	}
+
+	void FGenericProfileTreeBuilder::AddConfigurationProperty( FLaunchProfileTreeNode& HeadingNode )
+	{
+		HeadingNode.AddWidget( LOCTEXT("ConfigurationLabel", "Configuration"), 
+			SNew(SCustomLaunchLexToStringCombo<EBuildConfiguration>)
+			.OnSelectionChanged(this, &FGenericProfileTreeBuilder::SetBuildConfiguration)
+			.SelectedItem(this, &FGenericProfileTreeBuilder::GetBuildConfiguration)
+			.Items( TArray<EBuildConfiguration>({ EBuildConfiguration::Debug, EBuildConfiguration::Development, EBuildConfiguration::Test, EBuildConfiguration::Shipping }) )
+		);
+	}
+
+	void FGenericProfileTreeBuilder::AddContentSchemeProperty( FLaunchProfileTreeNode& HeadingNode )
+	{
+		HeadingNode.AddWidget( LOCTEXT("ContentSchemeLabel", "Content Scheme"), 
+			SNew(SCustomLaunchContentSchemeCombo)
+			.OnSelectionChanged(this, &FGenericProfileTreeBuilder::SetContentScheme)
+			.SelectedContentScheme_Lambda( [this]() { return ContentScheme; } )
+			.IsContentSchemeAvailable(this, &FGenericProfileTreeBuilder::IsContentSchemeAvailable)
+			.Font(FCoreStyle::Get().GetFontStyle("SmallFont"))
+		);
+	}
+
+	void FGenericProfileTreeBuilder::AddCompressPakFilesProperty( FLaunchProfileTreeNode& HeadingNode )
+	{
+		// in this example we are adding a single boolean instead of a custom widget. the struct parameter defines how the value is accessed
+		// the ForPak is the lambda functions created in the constructor and is again aimed at readability
+		HeadingNode.AddBoolean( LOCTEXT("CompressPakFilesLabel", "Compress Pak Files"),
+			{
+				.GetValue = [this]()			{ return Profile->IsCompressed(); },
+				.SetValue = [this](bool bValue)	{ Profile->SetCompressed(bValue); },
+				.GetDefaultValue = [this]()		{ return DefaultProfile->IsCompressed(); },
+				.IsVisible = ForPak,
+			}
+		);
+	}
+
+	void FGenericProfileTreeBuilder::AddUseIoStoreProperty( FLaunchProfileTreeNode& HeadingNode )
+	{
+		HeadingNode.AddBoolean( LOCTEXT("UseIoStoreLabel", "Use Io Store"),
+			{
+				.GetValue = [this]()			{ return Profile->IsUsingIoStore(); },
+				.SetValue = [this](bool bValue)	{ Profile->SetUseIoStore(bValue); },
+				.GetDefaultValue = [this]()		{ return DefaultProfile->IsUsingIoStore(); },
+				.IsVisible = ForPak,
+			}
+		);
+	}
+
+	void FGenericProfileTreeBuilder::AddGenerateChunksProperty( FLaunchProfileTreeNode& HeadingNode )
+	{
+		HeadingNode.AddBoolean( LOCTEXT("GenerateChunksLabel", "Generate Chunks"),
+			{
+				.GetValue = [this]()			{ return Profile->IsGeneratingChunks(); },
+				.SetValue = [this](bool bValue)	{ Profile->SetGenerateChunks(bValue); },
+				.GetDefaultValue = [this]()		{ return DefaultProfile->IsGeneratingChunks(); },
+				.IsVisible = ForPak,
+			}
+		);
+	}
+
+	void FGenericProfileTreeBuilder::AddImportZenSnapshotProperty( FLaunchProfileTreeNode& HeadingNode )
+	{
+		HeadingNode.AddBoolean( LOCTEXT("ImportZenSnapshotLabel", "Import Best Match Zen Snapshot"),
+			{
+				.GetValue = [this]()			{ return Profile->IsImportingZenSnapshot(); },
+				.SetValue = [this](bool bValue)	{ Profile->SetImportingZenSnapshot(bValue); },
+				.GetDefaultValue = [this]()		{ return DefaultProfile->IsImportingZenSnapshot(); },
+				.IsVisible = ForContent,
+			}
+		);
+	}
+
+	void FGenericProfileTreeBuilder::AddZenPakStreamingPathProperty( FLaunchProfileTreeNode& HeadingNode )
+	{
+		HeadingNode.AddDirectoryString( LOCTEXT("ZenPakStreamingPathLabel", "Zen Pak Streaming Path"),
+			{
+				.GetValue = [this]()				{ return Profile->GetZenPakStreamingPath(); },
+				.SetValue = [this](FString Value)	{ Profile->SetZenPakStreamingPath(Value); },
+				.GetDefaultValue = EmptyString,
+				.IsVisible = ForZenWS,
+			}
+		);
+	}
+
+	void FGenericProfileTreeBuilder::AddIncrementalCookProperty( FLaunchProfileTreeNode& HeadingNode )
+	{
+		auto GetDisplayName = []( ELauncherProfileIncrementalCookMode::Type Mode)
+		{
+			switch (Mode)
+			{
+				case ELauncherProfileIncrementalCookMode::None:                    return LOCTEXT("IncrementalCookNone",                 "None");
+				case ELauncherProfileIncrementalCookMode::ModifiedOnly:            return LOCTEXT("IncrementalCookModified",             "Modified Only (legacy)");
+				case ELauncherProfileIncrementalCookMode::ModifiedAndDependencies: return LOCTEXT("IncrementalCookModifiedDependencies", "Modified & Dependencies (recommended)");
+			}
+			return FText::GetEmpty();
+		};
+
+		auto GetToolTip = []( ELauncherProfileIncrementalCookMode::Type Mode)
+		{
+			switch (Mode)
+			{
+				case ELauncherProfileIncrementalCookMode::None:                    return LOCTEXT("IncrementalCookNoneTip",                 "This will try to cook everything");
+				case ELauncherProfileIncrementalCookMode::ModifiedOnly:            return LOCTEXT("IncrementalCookModifiedTip",             "This will only try to cook any modified assets but won't try to cook anything that depends on these assets. This is the old, legacy option and is faster but unreliable");
+				case ELauncherProfileIncrementalCookMode::ModifiedAndDependencies: return LOCTEXT("IncrementalCookModifiedDependenciesTip", "This will try to cook any modified assets and those assets that depend on them. This is the new method and is much more reliable but is a little slower");
+			}
+			return FText::GetEmpty();
+		};
+
+		HeadingNode.AddWidget( LOCTEXT("IncrementalCookLabel", "Incremental Cook"),
+			{
+				.IsDefault = [this]()		{ return GetIncrementalCookMode() == DefaultProfile->GetIncrementalCookMode(); },
+				.SetToDefault = [this]()	{ SetIncrementalCookMode(DefaultProfile->GetIncrementalCookMode()); },
+				.IsVisible = ForEnabledCooked,
+			},
+			SNew(SCustomLaunchCombo<ELauncherProfileIncrementalCookMode::Type>)
+			.OnSelectionChanged(this, &FGenericProfileTreeBuilder::SetIncrementalCookMode)
+			.SelectedItem(this, &FGenericProfileTreeBuilder::GetIncrementalCookMode)
+			.GetDisplayName_Lambda(GetDisplayName)
+			.GetItemToolTip_Lambda(GetToolTip)
+			.Items( TArray<ELauncherProfileIncrementalCookMode::Type>(
+				{
+					ELauncherProfileIncrementalCookMode::None,
+					ELauncherProfileIncrementalCookMode::ModifiedOnly,
+					ELauncherProfileIncrementalCookMode::ModifiedAndDependencies,
+				})
+			)
+		);
+	}
+
+
+
+	void FGenericProfileTreeBuilder::AddCookProperty( FLaunchProfileTreeNode& HeadingNode )
+	{
+		HeadingNode.AddBoolean( LOCTEXT("CookLabel", "Cook Content"),
+			{
+				.GetValue = [this]()			{ return GetCook(); },
+				.SetValue = [this](bool bValue)	{ SetCook(bValue); },
+				.GetDefaultValue = [this]()		{ return GetCook(DefaultProfile); },
+				.IsVisible = ForCooked
+			}
+		);
+	}
+
+	void FGenericProfileTreeBuilder::AddMapsToCookProperty( FLaunchProfileTreeNode& HeadingNode )
+	{
+		HeadingNode.AddWidget( LOCTEXT("MapsToCookLabel", "Maps To Cook"), 
+			{
+				.IsVisible = ForEnabledCooked,
+			},
+			CreateMapListWidget()
+		);
+	}
+
+	void FGenericProfileTreeBuilder::AddAdditionalCookerOptionsProperty( FLaunchProfileTreeNode& HeadingNode )
+	{
+		HeadingNode.AddString( LOCTEXT("AdditionalCookerOptionsLabel", "Additional Cooker Options"), 
+			{
+				.GetValue = [this]()				{ return Profile->GetCookOptions(); },
+				.SetValue = [this](FString Value)	{ Profile->SetCookOptions( Value ); },
+				.GetDefaultValue = EmptyString,
+				.IsVisible = ForEnabledCooked,
+			}
+		);
+	}
+
+	void FGenericProfileTreeBuilder::AddBuildProperty( FLaunchProfileTreeNode& HeadingNode )
+	{
+		HeadingNode.AddBoolean( LOCTEXT("BuildLabel", "Build the game"),
+			{
+				.GetValue = [this]()			{ return GetBuild(); },
+				.SetValue = [this](bool bValue)	{ SetBuild(bValue); },
+				.GetDefaultValue = [this]()		{ return GetBuild(DefaultProfile); },
+			}
+		);
+	}
+
+	void FGenericProfileTreeBuilder::AddForceBuildProperty( FLaunchProfileTreeNode& HeadingNode )
+	{
+		HeadingNode.AddBoolean( LOCTEXT("ForceBuildLabel", "Build even if a pre-built target exists"),
+			{
+				.GetValue = [this]()			{ return GetForceBuild(); },
+				.SetValue = [this](bool bValue)	{ SetForceBuild(bValue); },
+				.GetDefaultValue = [this]()		{ return GetForceBuild(DefaultProfile); },
+				.IsEnabled = ForCode,
+			}
+		);
+	}
+
+	void FGenericProfileTreeBuilder::AddBuidUATProperty( FLaunchProfileTreeNode& HeadingNode )
+	{
+		HeadingNode.AddBoolean( LOCTEXT("BuildUATLabel", "Build UAT"),
+			{
+				.GetValue = [this]()			{ return Profile->IsBuildingUAT(); },
+				.SetValue = [this](bool bValue)	{ Profile->SetBuildUAT(bValue); },
+				.GetDefaultValue = [this]()		{ return DefaultProfile->IsBuildingUAT(); },
+				.IsEnabled = ForCode,
+			}
+		);
+	}
+
+	void FGenericProfileTreeBuilder::AddArchitectureProperty( FLaunchProfileTreeNode& HeadingNode )
+	{
+		HeadingNode.AddWidget( LOCTEXT("ArchitectureLabel", "Architecture"),
+			{
+				.IsDefault = [this]()		{ return GetArchitecture().IsEmpty(); },
+				.SetToDefault = [this]()	{ SetArchitecture(FString()); },
+				.IsVisible = [this]()       { return CachedArchitectures.Num() > 0; }
+			},
+			SNew(SCustomLaunchStringCombo)
+			.OnSelectionChanged( this, &FGenericProfileTreeBuilder::SetArchitecture )
+			.SelectedItem( this, &FGenericProfileTreeBuilder::GetArchitecture)
+			.GetDisplayName( this, &FGenericProfileTreeBuilder::GetArchitectureDisplayName )
+			.Items_Lambda( [this]() { return CachedArchitectures;} )
+		);
+	}
+
+
+
+
+	void FGenericProfileTreeBuilder::AddStagingDirectoryProperty( FLaunchProfileTreeNode& HeadingNode )
+	{
+		HeadingNode.AddDirectoryString( LOCTEXT("CustomStagingPathLabel", "Custom Stage Directory"),
+			{
+				.GetValue = [this]()				{ return Profile->GetPackageDirectory(); },
+				.SetValue = [this](FString Value)	{ Profile->SetPackageDirectory(Value); },
+				.GetDefaultValue = EmptyString,
+				.IsVisible = ForCooked,
+			}
+		);
+	}
+
+	void FGenericProfileTreeBuilder::AddArchiveBuildProperty( FLaunchProfileTreeNode& HeadingNode )
+	{
+		HeadingNode.AddBoolean( LOCTEXT("ArchiveBuildLabel", "Archive Build"),
+			{
+				.GetValue = [this]()			{ return Profile->IsArchiving(); },
+				.SetValue = [this](bool bValue)	{ Profile->SetArchive(bValue); },
+				.GetDefaultValue = [this]()		{ return DefaultProfile->IsArchiving(); },
+				.IsVisible = ForCooked,
+			}
+		);
+	}
+
+
+	void FGenericProfileTreeBuilder::AddArchiveBuildDirectoryProperty( FLaunchProfileTreeNode& HeadingNode )
+	{
+		HeadingNode.AddDirectoryString( LOCTEXT("ArchivePathLabel", "Archive Directory"),
+			{
+				.GetValue = [this]()				{ return Profile->GetArchiveDirectory(); },
+				.SetValue = [this](FString Value)	{ Profile->SetArchiveDirectory(Value); },
+				.GetDefaultValue = EmptyString,
+				.IsVisible = ForCooked,				
+				.IsEnabled = [this]()				{ return Profile->IsArchiving(); }
+			}
+		);
+	}
+
+
+
+	void FGenericProfileTreeBuilder::AddDeployProperty( FLaunchProfileTreeNode& HeadingNode )
+	{
+		HeadingNode.AddBoolean( LOCTEXT("DeployLabel", "Deploy To Device"),
+			{
+				.GetValue = [this]()			{ return GetDeployToDevice(); },
+				.SetValue = [this](bool bValue)	{ SetDeployToDevice(bValue); },
+				.GetDefaultValue = [this]()		{ return GetDeployToDevice(DefaultProfile); },
+				.IsVisible = [this]()			{ return (ContentScheme != EContentScheme::CookOnTheFly); }
+			}
+		);
+	}
+
+	void FGenericProfileTreeBuilder::AddIncrementalDeployProperty( FLaunchProfileTreeNode& HeadingNode )
+	{
+		HeadingNode.AddBoolean( LOCTEXT("DeployModifiedLabel", "Only Deploy Modified Content"),
+			{
+				.GetValue = [this]()			{ return Profile->IsDeployingIncrementally(); },
+				.SetValue = [this](bool bValue)	{ Profile->SetIncrementalDeploying(bValue); },
+				.GetDefaultValue = [this]()		{ return DefaultProfile->IsDeployingIncrementally(); },
+				.IsEnabled = ForDeployment,
+			}
+		);
+	}
+
+	void FGenericProfileTreeBuilder::AddTargetDeviceProperty( FLaunchProfileTreeNode& HeadingNode )
+	{
+		HeadingNode.AddWidget( LOCTEXT("TargetDeviceLabel", "Target Device"), 
+			{
+				.IsEnabled = ForDeployment,
+			},
+			CreateDeployDeviceWidget()
+		);
+	}
+
+
+	void FGenericProfileTreeBuilder::AddRunProperty( FLaunchProfileTreeNode& HeadingNode )
+	{
+		HeadingNode.AddBoolean( LOCTEXT("RunLabel", "Run"),
+			{
+				.GetValue = [this]()			{ return GetIsRunning(); },
+				.SetValue = [this](bool bValue)	{ SetIsRunning(bValue); },
+				.GetDefaultValue = [this]()		{ return GetIsRunning(DefaultProfile); },
+			}
+		);
+	}
+
+
+	void FGenericProfileTreeBuilder::AddInitialMapProperty( FLaunchProfileTreeNode& HeadingNode )
+	{
+		// todo: custom map picker 
+		HeadingNode.AddWidget( LOCTEXT("InitialMapLabel", "Initial Map"), 
+			{
+				.IsDefault = [this]()		{ return Profile->GetDefaultLaunchRole()->GetInitialMap().IsEmpty(); },
+				.SetToDefault = [this]()	{ Profile->GetDefaultLaunchRole()->SetInitialMap(FString()); InitalMapCombo->SetSelectedItem(GetInitialMap()); },
+				.IsEnabled = ForRun,
+			},
+			SAssignNew(InitalMapCombo, SSearchableComboBox)
+			.OptionsSource(&CachedStartupMaps)
+			.OnSelectionChanged( this, &FGenericProfileTreeBuilder::OnInitialMapChanged )
+			.OnGenerateWidget( this, &FGenericProfileTreeBuilder::OnGenerateComboWidget )
+			.Content()
+			[
+				SNew(STextBlock)
+				.Text_Lambda( [this]() { return FText::FromString( *GetInitialMap() ); } )
+				.Font(FCoreStyle::Get().GetFontStyle("SmallFont"))
+			]
+		);
+	}
+
+
+	void FGenericProfileTreeBuilder::AddCommandLineProperty( FLaunchProfileTreeNode& HeadingNode )
+	{
+		HeadingNode.AddCommandLineString( LOCTEXT("CommandLineLabel", "Additional Command Line"), 
+			{
+				.GetValue = [this]()				{ return GetCommandLine(); },
+				.SetValue = [this](FString Value)	{ SetCommandLine(Value); },
+				.GetDefaultValue = EmptyString,
+				.IsEnabled = ForRun
+			}
+		);
+	}
+
+
+
+
+
+
+
+
+
+
+
+
+
+
+
+
+
+
+
+
+
+
+
+	void FGenericProfileTreeBuilder::OnPropertyChanged()
+	{
+		if (ProfileType == EProfileType::Basic)
+		{
+			// do not save basic profiles - they're transient
+		}
+		else
+		{
+			Model->GetProfileManager()->SaveJSONProfile(Profile);
+		}
+	}
+
+
+	void FGenericProfileTreeBuilder::CacheStartupMapList()
+	{
+		CachedStartupMaps.Reset();
+		CachedStartupMaps.Add(MakeShared<FString>());
+		for ( const FString& Map : FGameProjectHelper::GetAvailableMaps(Profile->GetProjectBasePath(), Profile->SupportsEngineMaps(), true))
+		{
+			CachedStartupMaps.Add(MakeShared<FString>(Map));
+		}
+
+		if (InitalMapCombo.IsValid())
+		{
+			InitalMapCombo->SetSelectedItem( GetInitialMap(), ESelectInfo::Direct );
+			InitalMapCombo->RefreshOptions();
+		}
+	}
+
+	void FGenericProfileTreeBuilder::CacheArchitectures()
+	{
+		CachedArchitectures.Reset();
+
+		const PlatformInfo::FTargetPlatformInfo* PlatformInfo = Model->GetPlatformInfo(Profile);
+		if (PlatformInfo != nullptr)
+		{
+			const ITargetPlatform* TargetPlatform = GetTargetPlatformManager()->FindTargetPlatform(PlatformInfo->IniPlatformName);
+			if (TargetPlatform)
+			{
+				TargetPlatform->GetPossibleArchitectures(CachedArchitectures);
+				if (CachedArchitectures.Num() > 0)
+				{
+					CachedArchitectures.Insert(FString(), 0); // empty string for "project default" option
+				}
+			}
+		}
+	}
+
+
+
+	TSharedRef<SWidget> FGenericProfileTreeBuilder::OnGenerateComboWidget( TSharedPtr<FString> InComboString )
+	{
+		return SNew(STextBlock)
+			.Text(InComboString.IsValid() ? FText::FromString(*InComboString) : FText::GetEmpty() )
+			.Font(FCoreStyle::Get().GetFontStyle("SmallFont"));
+
+	}
+
+
+
+	void FGenericProfileTreeBuilder::SetSelectedPlatforms( TArray<FString> SelectedPlatforms )
+	{
+		Profile->ClearCookedPlatforms();
+
+		if (bUseFriendlyBuildTargetSelection)
+		{
+			FTargetInfo BuildTargetInfo = FModel::GetBuildTargetInfo(GetBuildTarget(), GetProjectPath());
+
+			for (const FString& Platform : SelectedPlatforms)
+			{
+				Profile->AddCookedPlatform( FModel::GetBuildTargetPlatformName(Platform, BuildTargetInfo));
+			}
+		}
+		else
+		{
+			for (const FString& Platform : SelectedPlatforms)
+			{
+				Profile->AddCookedPlatform(Platform);
+			}
+		}
+		OnPropertyChanged();
+
+		DeployDeviceListView->OnSelectedPlatformChanged();
+
+		CacheArchitectures();
+	}
+
+	TArray<FString> FGenericProfileTreeBuilder::GetSelectedPlatforms() const
+	{
+		TArray<FString> Platforms;
+
+		if (bUseFriendlyBuildTargetSelection)
+		{
+			for (const FString& Platform : Profile->GetCookedPlatforms())
+			{
+				Platforms.Add(FModel::GetVanillaPlatformName(Platform));
+			}
+		}
+		else
+		{
+			Platforms = Profile->GetCookedPlatforms();
+		}
+
+		return MoveTemp(Platforms);
+	}
+
+
+	FString FGenericProfileTreeBuilder::GetProjectPath() const
+	{
+		if (ProfileType == EProfileType::Basic)
+		{
+			if (Profile->HasProjectSpecified())
+			{
+				return Profile->GetProjectPath();
+			}
+			else
+			{
+				return Model->GetProfileManager()->GetProjectPath();
+			}
+		}
+		else
+		{
+			if (Profile->HasProjectSpecified())
+			{
+				return Profile->GetProjectPath();
+			}
+
+			return FString();
+		}
+	}
+
+	void FGenericProfileTreeBuilder::SetProjectName(FString ProjectPath)
+	{
+		Profile->SetProjectSpecified(!ProjectPath.IsEmpty());
+		Profile->SetProjectPath(ProjectPath);
+
+		if (bUseFriendlyBuildTargetSelection && ProfileType != EProfileType::Basic )
+		{
+			Model->UpdateCookedPlatformsFromBuildTarget(Profile);
+		}
+
+		OnPropertyChanged();
+
+		if (MapListView.IsValid())
+		{
+			MapListView->RefreshMapList();
+		}
+
+		CacheStartupMapList();
+	}
+
+	bool FGenericProfileTreeBuilder::HasProject() const
+	{
+		return Profile->HasProjectSpecified();
+	}
+
+
+	FString FGenericProfileTreeBuilder::GetBuildTarget() const
+	{
+		if (ProfileType == EProfileType::Basic)
+		{
+			if (Profile->HasBuildTargetSpecified())
+			{
+				return Profile->GetBuildTarget();
+			}
+			else
+			{
+				return Model->GetProfileManager()->GetBuildTarget();
+			}
+		}
+		else
+		{
+			return Profile->GetBuildTarget();
+		}
+	}
+
+	void FGenericProfileTreeBuilder::SetBuildTarget(FString BuildTarget)
+	{
+		Profile->SetBuildTargetSpecified(!BuildTarget.IsEmpty());
+		Profile->SetBuildTarget(BuildTarget);
+	
+		if (ProfileType == EProfileType::Basic )
+		{
+			Model->UpdatedCookedPlatformsFromDeployDeviceProxy(Profile);
+		}
+		else if (bUseFriendlyBuildTargetSelection)
+		{
+			Model->UpdateCookedPlatformsFromBuildTarget(Profile);
+		}
+
+		CachedBuildTargetType = Model->GetBuildTargetInfo(Profile).Type;
+		OnPropertyChanged();
+	}
+
+	TArray<EBuildTargetType> FGenericProfileTreeBuilder::GetSupportedBuildTargetTypes() const
+	{
+		TArray<EBuildTargetType> Result;
+
+		if (Profile->GetCookedPlatforms().Num() > 0)
+		{
+			const PlatformInfo::FTargetPlatformInfo* PlatformInfo = PlatformInfo::FindPlatformInfo(FName(Profile->GetCookedPlatforms()[0]));
+			if (PlatformInfo)
+			{
+				if (bUseFriendlyBuildTargetSelection)
+				{
+					Result.Add(PlatformInfo->VanillaInfo->PlatformType);
+
+					for (const PlatformInfo::FTargetPlatformInfo* PlatformFlavorInfo : PlatformInfo->VanillaInfo->Flavors)
+					{
+						Result.AddUnique(PlatformFlavorInfo->PlatformType);
+					}
+				}
+				else
+				{
+					Result.Add(PlatformInfo->PlatformType);
+				}
+			}
+		}
+
+		return MoveTemp(Result);
+	}
+
+	void FGenericProfileTreeBuilder::SetBuildConfiguration(EBuildConfiguration BuildConfiguration)
+	{
+		Profile->SetBuildConfiguration(BuildConfiguration);
+		OnPropertyChanged();
+	}
+
+	EBuildConfiguration FGenericProfileTreeBuilder::GetBuildConfiguration() const
+	{
+		return Profile->GetBuildConfiguration();
+	}
+
+
+
+	void FGenericProfileTreeBuilder::RefreshContentScheme()
+	{
+		EContentScheme CurrentContentScheme = Model->DetermineProfileContentScheme(Profile);
+		SetContentScheme(CurrentContentScheme);
+	}
+
+	void FGenericProfileTreeBuilder::SetContentScheme(EContentScheme InContentScheme)
+	{
+		ELauncherProfileDeploymentModes::Type DeploymentMode = GetDeployToDevice() ? ELauncherProfileDeploymentModes::CopyToDevice : ELauncherProfileDeploymentModes::DoNotDeploy;
+		Model->SetProfileContentScheme(InContentScheme, Profile, bShouldCook, DeploymentMode );
+		OnPropertyChanged();
+
+		// refresh the cached content scheme in case the option that was selected is not available
+		ContentScheme = Model->DetermineProfileContentScheme(Profile);
+	}
+
+	bool FGenericProfileTreeBuilder::IsContentSchemeAvailable(EContentScheme InContentScheme, FText& OutReason) const
+	{
+		FProjectSettings ProjectSettings = Model->GetProjectSettings(Profile);
+
+		// basic launch is aimed at launching the current content, not an external build
+		if (ProfileType == EProfileType::Basic && InContentScheme == EContentScheme::ZenPakStreaming)
+		{
+			// don't set a reason - just hide the item for Basic Launch
+			return false;
+		}
+
+		// loose files can't be selected if the project is using zen store because there's no way to opt-out of Zen Store from the UAT command line.
+		if (InContentScheme == EContentScheme::LooseFiles && ProjectSettings.bUseZenStore)
+		{
+			OutReason = LOCTEXT("NoLooseFilesReason", "Loose Files cannot be used when Zen Store is enabled in Project Settings");
+			return false;
+		}
+
+		// don't show zen pak streaming option if it isn't going to be set up automatically by UAT
+		// @todo: could potentially look at shelling out to Zen & querying if we have a dynamic workspace? no support for async config queries in this tool yet though
+		if (InContentScheme == EContentScheme::ZenPakStreaming && !ProjectSettings.bHasAutomaticZenPakStreamingWorkspaceCreation)
+		{
+			OutReason = LOCTEXT("NoZenPakReason", "Automatic Zen Pak streaming workspace creation has not been enabled in Project Settings");
+			return false;
+		}
+	
+		// cannot launch via Zen if we're targeting a remote device and Zen isn't accepting external connections
+		if ((InContentScheme == EContentScheme::ZenStreaming || InContentScheme == EContentScheme::ZenPakStreaming) 
+			&& !ProjectSettings.bAllowRemoteNetworkService && !Model->IsHostPlatform(Profile))
+		{
+			OutReason = LOCTEXT("NoZenReason", "Zen Streaming to a remote device requires AllowRemoteNetworkService");
+			return false;
+		}
+
+		return true;
+	}
+
+
+
+	FString FGenericProfileTreeBuilder::GetCommandLine() const
+	{
+		// get unified command line from these two fields. the first is presented in "Build" for old Project Launcher, and the latter is presented in "Launch" for old Project Launcher)
+		// when we save back to the profile, this will be stored just in the "Build" one for clarity (because multiple roles are not supported in old or new project launcher)
+		FString CommandLine = Profile->GetAdditionalCommandLineParameters().TrimStartAndEnd() + TEXT(" ") + Profile->GetDefaultLaunchRole()->GetUATCommandLine().TrimStartAndEnd();
+		CommandLine.TrimStartAndEndInline();
+
+		return CommandLine;
+	}
+
+	void FGenericProfileTreeBuilder::SetCommandLine( const FString& NewCommandLine )
+	{
+		Profile->SetAdditionalCommandLineParameters(NewCommandLine);
+		Profile->GetDefaultLaunchRole()->SetCommandLine(TEXT(""));
+		OnPropertyChanged();
+	}
+
+
+
+
+	BEGIN_SLATE_FUNCTION_BUILD_OPTIMIZATION
+	TSharedRef<SWidget> FGenericProfileTreeBuilder::CreateMapListWidget()
+	{
+		MapListView = 
+			SNew(SCustomLaunchMapListView)
+			.OnSelectionChanged(this, &FGenericProfileTreeBuilder::SetMapsToCook)
+			.SelectedMaps(this, &FGenericProfileTreeBuilder::GetMapsToCook)
+			.ProjectPath(this, &FGenericProfileTreeBuilder::GetProjectPath)
+		;
+
+		return SNew(SVerticalBox)
+
+			// map option controls
+			+SVerticalBox::Slot()
+			.AutoHeight()
+			.Padding(0,2)
+			[
+				SNew(SHorizontalBox)
+
+				// map option
+				+SHorizontalBox::Slot()
+				.Padding(0, 0)
+				.AutoWidth()
+				.VAlign(VAlign_Center)
+				[
+					SNew(SSegmentedControl<EMapOption>)
+					.Value( this, &FGenericProfileTreeBuilder::GetMapOption)
+					.OnValueChanged(this, &FGenericProfileTreeBuilder::SetMapOption)
+
+					+SSegmentedControl<EMapOption>::Slot(EMapOption::Startup)
+					.HAlign(HAlign_Center)
+					.VAlign(VAlign_Center)
+					[
+						SNew(STextBlock)
+						.Text(LOCTEXT("StartupMapsLabel", "Startup Maps"))
+						.Font(FCoreStyle::Get().GetFontStyle("SmallFont"))
+					]
+
+					+SSegmentedControl<EMapOption>::Slot(EMapOption::Selected)
+					.HAlign(HAlign_Center)
+					.VAlign(VAlign_Center)
+					[
+						SNew(STextBlock)
+						.Text(LOCTEXT("SelectedMapsLabel", "Selected Maps"))
+						.Font(FCoreStyle::Get().GetFontStyle("SmallFont"))
+					]
+				]
+
+				// map selector controls (search etc)
+				+SHorizontalBox::Slot()
+				.Padding(8, 0)
+				.FillWidth(1)
+				.VAlign(VAlign_Center)
+				[
+					SNew(SBox)
+					.Visibility_Lambda([this](){ return GetMapOption() == EMapOption::Selected  ? EVisibility::Visible : EVisibility::Collapsed; } )
+					[
+						MapListView->MakeControlsWidget()
+					]
+				]
+			]
+
+			// map list
+			+SVerticalBox::Slot()
+			.FillHeight(1)
+			[
+				SNew(SVerticalResizeBox)
+				.Visibility_Lambda([this](){ return (GetMapOption() == EMapOption::Selected) ? EVisibility::Visible : EVisibility::Collapsed; } )
+				.HandleHeight(4)
+				.ContentHeight(this, &FGenericProfileTreeBuilder::GetMapListHeight)
+				.ContentHeightChanged(this, &FGenericProfileTreeBuilder::SetMapListHeight)
+				.HandleColor( FAppStyle::Get().GetSlateColor("Colors.Secondary").GetSpecifiedColor() )
+				[
+					MapListView.ToSharedRef()
+				]
+			]
+		;
+	}
+	END_SLATE_FUNCTION_BUILD_OPTIMIZATION
+
+
+	void FGenericProfileTreeBuilder::SetCook( bool bCook )
+	{
+		bShouldCook = bCook;
+		RefreshContentScheme();
+	}
+
+	bool FGenericProfileTreeBuilder::GetCook(ILauncherProfilePtr InProfile) const
+	{
+		if (InProfile == nullptr)
+		{
+			InProfile = Profile;
+		}
+		return (InProfile->GetCookMode() != ELauncherProfileCookModes::DoNotCook);
+	}
+
+
+	void FGenericProfileTreeBuilder::SetIncrementalCookMode( ELauncherProfileIncrementalCookMode::Type Mode )
+	{
+		Profile->SetIncrementalCookMode(Mode);
+		Profile->SetUnversionedCooking(Mode != ELauncherProfileIncrementalCookMode::ModifiedOnly); // should always use unversioned, except the legacy -iterativecook because it does not handle invalidation due to c++ changes.
+		OnPropertyChanged();
+	}
+
+	ELauncherProfileIncrementalCookMode::Type FGenericProfileTreeBuilder::GetIncrementalCookMode() const
+	{
+		return Profile->GetIncrementalCookMode();
+	}
+
+
+
+	void FGenericProfileTreeBuilder::SetMapsToCook(TArray<FString> MapsToCook)
+	{
+		Profile->ClearCookedMaps(); 
+		for (const FString& Map : MapsToCook)
+		{
+			Profile->AddCookedMap(Map);
+		}
+
+		OnPropertyChanged();
+	}
+
+	TArray<FString> FGenericProfileTreeBuilder::GetMapsToCook() const
+	{
+		return Profile->GetCookedMaps();
+	}
+
+	float FGenericProfileTreeBuilder::GetMapListHeight() const
+	{
+		return MapListHeight;
+	}
+
+	void FGenericProfileTreeBuilder::SetMapListHeight( float NewHeight )
+	{
+		static const float MinMapListHeight = 100.0f;
+
+		MapListHeight = FMath::Max(NewHeight, MinMapListHeight);
+	}
+
+	FGenericProfileTreeBuilder::EMapOption FGenericProfileTreeBuilder::GetMapOption() const
+	{
+		return MapOption;
+	}
+
+	void FGenericProfileTreeBuilder::SetMapOption( EMapOption NewMapOption )
+	{
+		bool bShow = (NewMapOption == EMapOption::Selected);
+
+		MapOption = NewMapOption;
+
+		if (bShow)
+		{
+			// restore the cooked maps again, if any
+			if (CachedMapsToCook.Num() > 0 && Profile->GetCookedMaps().Num() == 0)
+			{
+				SetMapsToCook(CachedMapsToCook);
+				CachedMapsToCook.Reset();
+			}
+		}
+		else
+		{
+			// to set the 'cook startup maps only', its necessary to remove all the cooked maps - take a copy of the values to allow it to be restored
+			CachedMapsToCook = Profile->GetCookedMaps();
+			SetMapsToCook(TArray<FString>());
+		}
+
+		OnPropertyChanged();
+
+		MapListView->RefreshMapList();
+	}
+
+
+
+
+	BEGIN_SLATE_FUNCTION_BUILD_OPTIMIZATION
+	TSharedRef<SWidget> FGenericProfileTreeBuilder::CreateDeployDeviceWidget()
+	{
+		if (ProfileType == EProfileType::Basic)
+		{
+			return SNew(SCustomLaunchDeviceListView)
+				.OnSelectionChanged(this, &FGenericProfileTreeBuilder::SetDeployDeviceIDs)
+				.SelectedDevices(this, &FGenericProfileTreeBuilder::GetDeployDeviceIDs)
+				.AllPlatforms(true)
+				.SingleSelect(true)
+			;
+		}
+		else
+		{
+			return SNew(SVerticalBox)
+
+				// device picker options
+				+SVerticalBox::Slot()
+				.AutoHeight()
+				.Padding(0,2)
+				[
+					SNew(SHorizontalBox)
+
+					+SHorizontalBox::Slot()
+					.Padding(0, 0)
+					.AutoWidth()
+					.VAlign(VAlign_Center)
+					[
+						SNew(SSegmentedControl<EDeployDeviceOption>)
+						.Value( this, &FGenericProfileTreeBuilder::GetDeployDeviceOption)
+						.OnValueChanged(this, &FGenericProfileTreeBuilder::SetDeployDeviceOption)
+
+						+SSegmentedControl<EDeployDeviceOption>::Slot(EDeployDeviceOption::Default)
+						.HAlign(HAlign_Center)
+						.VAlign(VAlign_Center)
+						[
+							SNew(STextBlock)
+							.Text(LOCTEXT("DefaultDeviceLabel", "Default Device"))
+							.Font(FCoreStyle::Get().GetFontStyle("SmallFont"))
+						]
+
+						+SSegmentedControl<EDeployDeviceOption>::Slot(EDeployDeviceOption::Selected)
+						.HAlign(HAlign_Center)
+						.VAlign(VAlign_Center)
+						[
+							SNew(STextBlock)
+							.Text(LOCTEXT("SelectedDevicesLabel", "Selected Devices"))
+							.Font(FCoreStyle::Get().GetFontStyle("SmallFont"))
+						]
+					]
+				]
+
+				// device picker list
+				+SVerticalBox::Slot()
+				.FillHeight(1)
+				[
+					SNew(SVerticalResizeBox)
+					.Visibility_Lambda([this](){ return (GetDeployDeviceOption() == EDeployDeviceOption::Selected) ? EVisibility::Visible : EVisibility::Collapsed; } )
+					.HandleHeight(4)
+					.ContentHeight(this, &FGenericProfileTreeBuilder::GetDeployDeviceListHeight)
+					.ContentHeightChanged(this, &FGenericProfileTreeBuilder::SetDeployDeviceListHeight)
+					.HandleColor( FAppStyle::Get().GetSlateColor("Colors.Secondary").GetSpecifiedColor() )
+					[
+						SAssignNew(DeployDeviceListView, SCustomLaunchDeviceListView)
+						.OnDeviceRemoved( this, &FGenericProfileTreeBuilder::OnDeviceRemoved)
+						.OnSelectionChanged(this, &FGenericProfileTreeBuilder::SetDeployDeviceIDs)
+						.SelectedDevices(this, &FGenericProfileTreeBuilder::GetDeployDeviceIDs)
+						.Platforms(this, &FGenericProfileTreeBuilder::GetSelectedPlatforms)
+					]
+				]
+			;
+		}
+	}
+	END_SLATE_FUNCTION_BUILD_OPTIMIZATION
+
+
+	void FGenericProfileTreeBuilder::SetDeployDeviceIDs(TArray<FString> DeployDeviceIDs)
+	{
+		if (Profile->GetDeployedDeviceGroup() == nullptr)
+		{
+			Profile->SetDeployedDeviceGroup(Model->GetProfileManager()->AddNewDeviceGroup());
+		}
+
+		Profile->GetDeployedDeviceGroup()->RemoveAllDevices(); 
+		for ( const FString& DeviceID : DeployDeviceIDs )
+		{
+			Profile->GetDeployedDeviceGroup()->AddDevice(DeviceID);
+		}
+
+		if (ProfileType == EProfileType::Basic) // @fixme: should this be done for custom profiles too?
+		{
+			Model->UpdatedCookedPlatformsFromDeployDeviceProxy(Profile);
+		}
+
+		OnPropertyChanged();
+	}
+
+	TArray<FString> FGenericProfileTreeBuilder::GetDeployDeviceIDs() const
+	{
+		return Profile->GetDeployedDeviceGroup()->GetDeviceIDs();
+	}
+
+	float FGenericProfileTreeBuilder::GetDeployDeviceListHeight() const
+	{
+		return DeployDeviceListHeight;
+	}
+
+	void FGenericProfileTreeBuilder::SetDeployDeviceListHeight( float NewHeight )
+	{
+		static const float MinDeployDeviceListHeight = 100.0f;
+
+		DeployDeviceListHeight = FMath::Max(NewHeight, MinDeployDeviceListHeight);
+	}
+
+	void FGenericProfileTreeBuilder::OnDeviceRemoved(FString DeviceID)
+	{
+		CachedDeployDeviceIDs.Remove(DeviceID);
+	}
+
+	FGenericProfileTreeBuilder::EDeployDeviceOption FGenericProfileTreeBuilder::GetDeployDeviceOption() const
+	{
+		return DeployDeviceOption;
+	}
+
+	void FGenericProfileTreeBuilder::SetDeployDeviceOption( EDeployDeviceOption NewDeployDeviceOption )
+	{
+		bool bShow = (NewDeployDeviceOption == EDeployDeviceOption::Selected);
+
+		DeployDeviceOption = NewDeployDeviceOption;
+
+		if (bShow)
+		{
+			// restore the deployed device list again, if any
+			if (CachedDeployDeviceIDs.Num() > 0 && Profile->GetDeployedDeviceGroup()->GetDeviceIDs().Num() == 0)
+			{
+				SetDeployDeviceIDs(Profile->GetDeployedDeviceGroup()->GetDeviceIDs());
+				CachedDeployDeviceIDs.Reset();
+			}
+		}
+		else
+		{
+			// to set the 'default' deploy option, its necessary to remove all the devices - take a copy of the values to allow it to be restored
+			CachedDeployDeviceIDs = Profile->GetDeployedDeviceGroup()->GetDeviceIDs();
+			SetDeployDeviceIDs(TArray<FString>());
+		}
+
+		OnPropertyChanged();
+
+		DeployDeviceListView->RefreshDeviceList();
+	}
+
+
+	void FGenericProfileTreeBuilder::SetBuild( bool bBuild )
+	{
+		if (!bBuild)
+		{
+			Profile->SetBuildMode(ELauncherProfileBuildModes::DoNotBuild);
+		}
+		else if (GetForceBuild(Profile))
+		{
+			Profile->SetBuildMode(ELauncherProfileBuildModes::Build);
+		}
+		else
+		{
+			Profile->SetBuildMode(ELauncherProfileBuildModes::Auto);
+		}
+
+		OnPropertyChanged();
+	}
+
+	bool FGenericProfileTreeBuilder::GetBuild(ILauncherProfilePtr InProfile) const
+	{
+		if (InProfile == nullptr)
+		{
+			InProfile = Profile;
+		}
+		return (InProfile->GetBuildMode() != ELauncherProfileBuildModes::DoNotBuild);
+	}
+
+	void FGenericProfileTreeBuilder::SetForceBuild( bool bForceBuild )
+	{
+		if (!GetBuild())
+		{
+			Profile->SetBuildMode(ELauncherProfileBuildModes::DoNotBuild);
+		}
+		else if (bForceBuild)
+		{
+			Profile->SetBuildMode(ELauncherProfileBuildModes::Build);
+		}
+		else
+		{
+			Profile->SetBuildMode(ELauncherProfileBuildModes::Auto);
+		}
+
+		OnPropertyChanged();
+	}
+
+	bool FGenericProfileTreeBuilder::GetForceBuild(ILauncherProfilePtr InProfile) const
+	{
+		if (InProfile == nullptr)
+		{
+			InProfile = Profile;
+		}
+		return (InProfile->GetBuildMode() == ELauncherProfileBuildModes::Build);
+	}
+
+	void FGenericProfileTreeBuilder::SetArchitecture( FString Architecture )
+	{
+		// clear existing architectures
+		TArray<FString> Empty;
+		Profile->SetServerArchitectures(Empty);
+		Profile->SetEditorArchitectures(Empty);
+		Profile->SetClientArchitectures(Empty);
+
+		// set new single architecture
+		if (!Architecture.IsEmpty())
+		{
+			TArray<FString> Single;
+			Single.Add(Architecture);
+
+			switch (CachedBuildTargetType)
+			{
+				case EBuildTargetType::Server: Profile->SetServerArchitectures(Single); break;
+				case EBuildTargetType::Editor: Profile->SetEditorArchitectures(Single); break;
+				default:                       Profile->SetClientArchitectures(Single); break;
+			}
+		}
+
+		OnPropertyChanged();
+	}
+
+	FString FGenericProfileTreeBuilder::GetArchitecture() const
+	{
+		TArray<FString> Architectures;
+
+		switch (CachedBuildTargetType)
+		{
+			case EBuildTargetType::Server: Architectures = Profile->GetServerArchitectures(); break;
+			case EBuildTargetType::Editor: Architectures = Profile->GetEditorArchitectures(); break;
+			default:                       Architectures = Profile->GetClientArchitectures(); break;
+		}
+
+		return (Architectures.Num() > 0) ? Architectures[0] : FString();
+	}
+
+	FText FGenericProfileTreeBuilder::GetArchitectureDisplayName( FString Architecture )
+	{
+		if (Architecture.IsEmpty())
+		{
+			return LOCTEXT("DefaultArchName", "Project Default");
+		}
+		else if (Architecture == FPlatformMisc::GetHostArchitecture() && Model->IsHostPlatform(Profile))
+		{
+			return FText::Format(LOCTEXT("HostArchLabel", "{0} (this platform)"), FText::FromString(Architecture) );
+		}
+		else
+		{
+			return FText::FromString(Architecture);
+		}
+	}
+
+
+
+	void FGenericProfileTreeBuilder::SetDeployToDevice( bool bDeployToDevice )
+	{
+		if (ContentScheme != EContentScheme::CookOnTheFly)
+		{
+			Profile->SetDeploymentMode(bDeployToDevice ? ELauncherProfileDeploymentModes::CopyToDevice : ELauncherProfileDeploymentModes::DoNotDeploy );
+			OnPropertyChanged();
+		}
+	}
+
+	bool FGenericProfileTreeBuilder::GetDeployToDevice(ILauncherProfilePtr InProfile) const
+	{
+		if (InProfile == nullptr)
+		{
+			InProfile = Profile;
+		}
+		return InProfile->GetDeploymentMode() != ELauncherProfileDeploymentModes::DoNotDeploy;
+	}
+
+
+	void FGenericProfileTreeBuilder::SetIsRunning( bool bRun )
+	{
+		Profile->SetLaunchMode(bRun ? ELauncherProfileLaunchModes::DefaultRole : ELauncherProfileLaunchModes::DoNotLaunch);
+	}
+
+	bool FGenericProfileTreeBuilder::GetIsRunning( ILauncherProfilePtr InProfile) const
+	{
+		if (InProfile == nullptr)
+		{
+			InProfile = Profile;
+		}
+	
+		return InProfile->GetLaunchMode() != ELauncherProfileLaunchModes::DoNotLaunch;
+	}
+
+	void FGenericProfileTreeBuilder::OnInitialMapChanged( TSharedPtr<FString> InitialMap, ESelectInfo::Type )
+	{
+		if (InitialMap.IsValid())
+		{
+			Profile->GetDefaultLaunchRole()->SetInitialMap(*InitialMap);
+		}
+		else
+		{
+			Profile->GetDefaultLaunchRole()->SetInitialMap(FString());
+		}
+
+		OnPropertyChanged();
+	}
+
+	TSharedPtr<FString> FGenericProfileTreeBuilder::GetInitialMap() const
+	{
+		FString InitialMap = Profile->GetDefaultLaunchRole()->GetInitialMap();
+		const TSharedPtr<FString>* FoundMapPtr = CachedStartupMaps.FindByPredicate( [InitialMap]( const TSharedPtr<FString>& Map )
+		{
+			return Map.IsValid() && (InitialMap == *Map);
+		});
+
+		if (FoundMapPtr)
+		{
+			return *FoundMapPtr;
+		}
+		else
+		{
+			return MakeShared<FString>();
+		}
+	}
+
+}
+
+#undef LOCTEXT_NAMESPACE
diff --git a/Engine/Plugins/Developer/ProjectLauncher/Source/ProjectLauncher/Public/Model/ProjectLauncherModel.h b/Engine/Plugins/Developer/ProjectLauncher/Source/ProjectLauncher/Public/Model/ProjectLauncherModel.h
new file mode 100644
--- /dev/null
+++ b/Engine/Plugins/Developer/ProjectLauncher/Source/ProjectLauncher/Public/Model/ProjectLauncherModel.h
@@ -0,0 +1,209 @@
+// Copyright Epic Games, Inc. All Rights Reserved.
+
+#pragma once
+
+#include "CoreTypes.h"
+#include "Containers/Array.h"
+#include "Delegates/Delegate.h"
+#include "ILauncherProfileManager.h"
+#include "Templates/SharedPointer.h"
+#include "Misc/DataDrivenPlatformInfoRegistry.h"
+
+
+
+namespace PlatformInfo
+{
+	struct FTargetPlatformInfo;
+};
+
+class ITargetDeviceProxyManager;
+class ITargetDeviceProxy;
+class ILauncher;
+struct FTargetInfo;
+
+namespace ProjectLauncher
+{
+	enum class EContentScheme : uint8
+	{
+		PakFiles,
+		ZenStreaming,
+		ZenPakStreaming,
+		DevelopmentPackage,
+		LooseFiles,
+		CookOnTheFly,
+		MAX
+	};
+	TArray<EContentScheme> GetAllContentSchemes();
+	FText GetContentSchemeDisplayName(EContentScheme ContentScheme);
+	FText GetContentSchemeToolTip(EContentScheme ContentScheme);
+	FText GetProfileLaunchErrorMessage(ILauncherProfilePtr Profile);
+	const TCHAR* LexToString( const ProjectLauncher::EContentScheme& ContentScheme);
+	bool LexTryParseString( ProjectLauncher::EContentScheme& OutContentScheme, const TCHAR* String );
+
+	enum class EProfileType : uint8
+	{
+		Invalid,
+		Basic,
+		Custom,
+		Advanced,
+	};
+
+	struct FLaunchLogMessage
+	{
+		TSharedRef<FString> Message;
+		ELogVerbosity::Type Verbosity;
+
+		FLaunchLogMessage(const FString& InMessage, ELogVerbosity::Type InVerbosity)
+			: Message(MakeShared<FString>(InMessage))
+			, Verbosity(InVerbosity)
+		{}
+	};
+
+	struct FProjectSettings
+	{
+		// whether this project is the one that is currently open in the editor
+		bool bIsCurrentEditorProject = false;
+
+		// whether Zen Store is enabled via the project settings - in this case, we can't use Loose Files
+		bool bUseZenStore = false;
+
+		// whether a workspace will be created automatically when using Zen Pak streaming
+		bool bHasAutomaticZenPakStreamingWorkspaceCreation = false;
+
+		// whether Zen Server is allowed to accept connections from remote machines (i.e. console devkits, phones etc)
+		bool bAllowRemoteNetworkService = false;
+	};
+
+
+	DECLARE_DELEGATE_OneParam(FOnProfileClicked, const ILauncherProfilePtr&)
+
+	DECLARE_MULTICAST_DELEGATE_TwoParams(FOnSelectedProfileChanged, const ILauncherProfilePtr& /*new*/, const ILauncherProfilePtr& /*old*/)
+	DECLARE_MULTICAST_DELEGATE(FOnSelectedProjectChanged)
+	DECLARE_MULTICAST_DELEGATE(FOnBuildTargetChanged)
+	DECLARE_MULTICAST_DELEGATE(FOnSelectedPlatformChanged)
+
+	extern bool bUseFriendlyBuildTargetSelection;
+
+	class PROJECTLAUNCHER_API FModel : public TSharedFromThis<FModel>
+	{
+	public:
+		FModel(const TSharedRef<ITargetDeviceProxyManager>& InDeviceProxyManager, const TSharedRef<ILauncher>& InLauncher, const TSharedRef<ILauncherProfileManager>& InProfileManager);
+		~FModel();
+
+
+		const TSharedRef<ITargetDeviceProxyManager>& GetDeviceProxyManager() const
+		{
+			return DeviceProxyManager;
+		}
+
+		const TSharedRef<ILauncher>& GetLauncher() const
+		{
+			return Launcher;
+		}
+
+		const TSharedRef<ILauncherProfileManager>& GetProfileManager() const
+		{
+			return ProfileManager;
+		}
+
+
+		const TArray<ILauncherProfilePtr>& GetAllProfiles() const
+		{
+			return AllProfiles;
+		}
+
+		ILauncherProfilePtr GetBasicLaunchProfile() const
+		{
+			return BasicLaunchProfile;
+		}
+
+		const ILauncherProfilePtr& GetSelectedProfile() const
+		{
+			return SelectedProfile;
+		}
+
+		FOnSelectedProfileChanged& OnProfileSelected()
+		{
+			return ProfileSelectedDelegate;
+		}
+
+		void SelectProfile(const ILauncherProfilePtr& Profile);
+		void SortProfiles();
+
+		EProfileType GetProfileType(const ILauncherProfileRef& Profile) const;
+		bool IsAdvancedProfile(const ILauncherProfileRef& Profile) const;
+		bool IsBasicLaunchProfile( const ILauncherProfilePtr& Profile ) const;
+
+		ILauncherProfileRef CreateCustomProfile( const TCHAR* Name );
+		ILauncherProfileRef CreateBasicLaunchProfile();
+		ILauncherProfilePtr CloneCustomProfile(const ILauncherProfileRef& Profile);
+
+		const ILauncherProfileRef GetDefaultBasicLaunchProfile() const;
+		const ILauncherProfileRef GetDefaultCustomLaunchProfile() const;
+
+		EContentScheme DetermineProfileContentScheme(const ILauncherProfileRef& Profile) const;
+		void SetProfileContentScheme(EContentScheme ContentScheme, const ILauncherProfileRef& Profile, bool bWantToCook = true, ELauncherProfileDeploymentModes::Type DefaultDeploymentMode = ELauncherProfileDeploymentModes::CopyToDevice);
+
+		TSharedPtr<FLaunchLogMessage> AddLogMessage( const FString& InMessage, ELogVerbosity::Type InVerbosity = ELogVerbosity::Log );
+		void ClearLogMessages();
+		int32 GetNumLogMessages() const { return LaunchLogMessages.Num(); }
+
+
+		static const PlatformInfo::FTargetPlatformInfo* GetPlatformInfo(const ILauncherProfilePtr& Profile);
+		static const PlatformInfo::FTargetPlatformInfo* GetPlatformInfo(FName PlatformName, const FTargetInfo& BuildTargetInfo);
+
+		static bool IsHostPlatform(const ILauncherProfilePtr& Profile);
+		static bool IsHostPlatform(FName PlatformName);
+
+		static FTargetInfo GetBuildTargetInfo( const FString& BuildTargetName, const FString& ProjectPath );
+		static FTargetInfo GetBuildTargetInfo( const ILauncherProfileRef& Profile );
+
+		static TSharedPtr<ITargetDeviceProxy> GetDeviceProxy(const ILauncherProfileRef& Profile);
+		static void UpdatedCookedPlatformsFromDeployDeviceProxy(const ILauncherProfileRef& Profile, TSharedPtr<ITargetDeviceProxy> DeviceProxy = nullptr);
+		static void UpdateCookedPlatformsFromBuildTarget(const ILauncherProfileRef& Profile);
+
+		static FString GetVanillaPlatformName( const FString& PlatformName );
+		static FString GetBuildTargetPlatformName( const FString& PlatformName, const FTargetInfo& BuildTargetInfo );
+
+		const FProjectSettings GetProjectSettings( const FString& ProjectPath );
+		const FProjectSettings GetProjectSettings( const ILauncherProfileRef& Profile );
+
+		const TCHAR* GetConfigSection() const;
+		const FString& GetConfigIni() const;
+
+	private:
+		FString ConfigFileName;
+
+		void LoadConfig();
+		void SaveConfig();
+
+		void HandleProfileManagerProfileAdded(const ILauncherProfileRef& Profile);
+		void HandleProfileManagerProfileRemoved(const ILauncherProfileRef& Profile);
+		void HandleDeviceProxyAdded(const TSharedRef<ITargetDeviceProxy>& DeviceProxy );
+		void HandleDeviceProxyRemoved(const TSharedRef<ITargetDeviceProxy>& DeviceProxy );
+
+		void ReadProjectSettingsFromConfig( FConfigCacheIni& InConfig, FProjectSettings& OutResult );
+
+		void OnModifyLaunchCommandLine( const ILauncherProfileRef& InProfile, FString& InOutCommandLine );
+
+		TSharedRef<ITargetDeviceProxyManager> DeviceProxyManager;
+		TSharedRef<ILauncher> Launcher;
+		TSharedRef<ILauncherProfileManager> ProfileManager;
+
+		ILauncherProfilePtr SelectedProfile;
+		TArray<ILauncherProfilePtr> AllProfiles;
+		ILauncherProfilePtr BasicLaunchProfile;
+		bool bHasSetBasicLaunchProfilePlatform = false;
+
+		friend class FLaunchLogTextLayoutMarshaller;
+		TArray<TSharedPtr<FLaunchLogMessage>> LaunchLogMessages;
+
+		TMap<FString,FProjectSettings> CachedProjectSettings;
+
+		FOnSelectedProfileChanged ProfileSelectedDelegate;
+		FOnSelectedPlatformChanged PlatformChangedDelegate;
+
+		ILauncherProfilePtr DefaultBasicLaunchProfile;
+		ILauncherProfilePtr DefaultCustomLaunchProfile;
+	};
+}
diff --git a/Engine/Plugins/Developer/ProjectLauncher/Source/ProjectLauncher/Public/ProfileTree/BasicProfileTreeBuilder.h b/Engine/Plugins/Developer/ProjectLauncher/Source/ProjectLauncher/Public/ProfileTree/BasicProfileTreeBuilder.h
new file mode 100644
--- /dev/null
+++ b/Engine/Plugins/Developer/ProjectLauncher/Source/ProjectLauncher/Public/ProfileTree/BasicProfileTreeBuilder.h
@@ -0,0 +1,35 @@
+// Copyright Epic Games, Inc. All Rights Reserved.
+
+#pragma once
+
+#include "ProfileTree/GenericProfileTreeBuilder.h"
+
+namespace ProjectLauncher
+{
+
+	class FBasicProfileTreeBuilder : public FGenericProfileTreeBuilder
+	{
+	public:
+		FBasicProfileTreeBuilder( const ILauncherProfileRef& Profile, const TSharedRef<FModel>& InModel );
+		virtual ~FBasicProfileTreeBuilder() = default;
+
+		virtual void Construct() override;
+
+		virtual FString GetName() const override
+		{
+			return TEXT("BasicProfile");
+		}
+	};
+
+
+
+
+	class PROJECTLAUNCHER_API FBasicProfileTreeBuilderFactory : public ILaunchProfileTreeBuilderFactory
+	{
+	public:
+		virtual TSharedPtr<ILaunchProfileTreeBuilder> TryCreateTreeBuilder( const ILauncherProfileRef& InProfile, const TSharedRef<FModel>& InModel ) override;
+
+		virtual bool IsProfileTypeSupported(EProfileType ProfileType ) const override { return ProfileType == EProfileType::Basic; }
+
+	};
+}
diff --git a/Engine/Plugins/Developer/ProjectLauncher/Source/ProjectLauncher/Public/ProfileTree/CustomProfileTreeBuilder.h b/Engine/Plugins/Developer/ProjectLauncher/Source/ProjectLauncher/Public/ProfileTree/CustomProfileTreeBuilder.h
new file mode 100644
--- /dev/null
+++ b/Engine/Plugins/Developer/ProjectLauncher/Source/ProjectLauncher/Public/ProfileTree/CustomProfileTreeBuilder.h
@@ -0,0 +1,33 @@
+// Copyright Epic Games, Inc. All Rights Reserved.
+
+#pragma once
+
+#include "ProfileTree/GenericProfileTreeBuilder.h"
+
+namespace ProjectLauncher
+{
+	class PROJECTLAUNCHER_API FCustomProfileTreeBuilder : public FGenericProfileTreeBuilder
+	{
+	public:
+		FCustomProfileTreeBuilder( const ILauncherProfileRef& Profile, const TSharedRef<FModel>& InModel );
+		virtual ~FCustomProfileTreeBuilder() = default;
+
+		virtual void Construct() override;
+
+		virtual FString GetName() const override
+		{
+			return TEXT("CustomProfile");
+		}
+	};
+
+
+	class PROJECTLAUNCHER_API FCustomProfileTreeBuilderFactory : public ILaunchProfileTreeBuilderFactory
+	{
+	public:
+		virtual TSharedPtr<ILaunchProfileTreeBuilder> TryCreateTreeBuilder( const ILauncherProfileRef& Profile, const TSharedRef<FModel>& InModel ) override;
+
+		virtual bool IsProfileTypeSupported(EProfileType ProfileType ) const override { return ProfileType == EProfileType::Custom; }
+
+	};
+
+}
diff --git a/Engine/Plugins/Developer/ProjectLauncher/Source/ProjectLauncher/Public/ProfileTree/GenericProfileTreeBuilder.h b/Engine/Plugins/Developer/ProjectLauncher/Source/ProjectLauncher/Public/ProfileTree/GenericProfileTreeBuilder.h
new file mode 100644
--- /dev/null
+++ b/Engine/Plugins/Developer/ProjectLauncher/Source/ProjectLauncher/Public/ProfileTree/GenericProfileTreeBuilder.h
@@ -0,0 +1,195 @@
+// Copyright Epic Games, Inc. All Rights Reserved.
+
+#pragma once
+
+#include "ProfileTree/ILaunchProfileTreeBuilder.h"
+#include "Model/ProjectLauncherModel.h"
+
+class SSearchableComboBox;
+class SCustomLaunchMapListView;
+class SCustomLaunchDeviceListView;
+namespace ESelectInfo { enum Type : int; }
+
+namespace ProjectLauncher
+{
+	/**
+	 * Base class for a profile tree builder that creates FLaunchProfileTreeData from a given ILauncherProfile.
+	 * 
+	 * Expected to be created by an instance of ILaunchProfileTreeBuilderFactory, for example:
+	 * 
+	 *    TSharedPtr<ILaunchProfileTreeBuilder> FMyProfileTreeBuilderFactory::TryCreateTreeBuilder( const ILauncherProfileRef& InProfile, const TSharedRef<FModel>& InModel )
+	 *    {
+	 *        return MakeShared<FMyProfileTreeBuilder>(InProfile, InModel);
+	 *    }
+	 */
+	class PROJECTLAUNCHER_API FGenericProfileTreeBuilder : public ILaunchProfileTreeBuilder, public TSharedFromThis<FGenericProfileTreeBuilder>
+	{
+	public:
+		FGenericProfileTreeBuilder( const ILauncherProfileRef& Profile, const ILauncherProfileRef& InDefaultProfile, const TSharedRef<FModel>& InModel );
+		virtual ~FGenericProfileTreeBuilder() = default;
+
+		virtual void Construct() override;
+
+		virtual FLaunchProfileTreeDataRef GetProfileTree() override
+		{
+			return TreeData;
+		}
+
+		virtual bool AllowExtensionsUI() const override
+		{
+			return true;
+		}
+
+	protected:
+		// default property creation functions
+		void AddProjectProperty( FLaunchProfileTreeNode& HeadingNode );
+		void AddTargetProperty( FLaunchProfileTreeNode& HeadingNode );
+		void AddPlatformProperty( FLaunchProfileTreeNode& HeadingNode );
+		void AddConfigurationProperty( FLaunchProfileTreeNode& HeadingNode );
+		void AddContentSchemeProperty( FLaunchProfileTreeNode& HeadingNode );
+		void AddCompressPakFilesProperty( FLaunchProfileTreeNode& HeadingNode );
+		void AddUseIoStoreProperty( FLaunchProfileTreeNode& HeadingNode );
+		void AddGenerateChunksProperty( FLaunchProfileTreeNode& HeadingNode );
+		void AddImportZenSnapshotProperty( FLaunchProfileTreeNode& HeadingNode );
+		void AddZenPakStreamingPathProperty( FLaunchProfileTreeNode& HeadingNode );
+		void AddIncrementalCookProperty( FLaunchProfileTreeNode& HeadingNode );
+		void AddCookProperty( FLaunchProfileTreeNode& HeadingNode );
+		void AddMapsToCookProperty( FLaunchProfileTreeNode& HeadingNode );
+		void AddAdditionalCookerOptionsProperty( FLaunchProfileTreeNode& HeadingNode );
+		void AddBuildProperty( FLaunchProfileTreeNode& HeadingNode );
+		void AddForceBuildProperty( FLaunchProfileTreeNode& HeadingNode );
+		void AddBuidUATProperty( FLaunchProfileTreeNode& HeadingNode );
+		void AddArchitectureProperty( FLaunchProfileTreeNode& HeadingNode );
+		void AddStagingDirectoryProperty( FLaunchProfileTreeNode& HeadingNode );
+		void AddArchiveBuildProperty( FLaunchProfileTreeNode& HeadingNode );
+		void AddArchiveBuildDirectoryProperty( FLaunchProfileTreeNode& HeadingNode );
+		void AddDeployProperty( FLaunchProfileTreeNode& HeadingNode );
+		void AddIncrementalDeployProperty( FLaunchProfileTreeNode& HeadingNode );
+		void AddTargetDeviceProperty( FLaunchProfileTreeNode& HeadingNode );
+		void AddRunProperty( FLaunchProfileTreeNode& HeadingNode );
+		void AddInitialMapProperty( FLaunchProfileTreeNode& HeadingNode );
+		void AddCommandLineProperty( FLaunchProfileTreeNode& HeadingNode );
+
+		// helper callbacks to simplify control enable/visibility
+		FLaunchProfileTreeNode::FGetBool ForPak;
+		FLaunchProfileTreeNode::FGetBool ForZenWS;
+		FLaunchProfileTreeNode::FGetBool ForCooked;
+		FLaunchProfileTreeNode::FGetBool ForEnabledCooked;
+		FLaunchProfileTreeNode::FGetBool ForContent;
+		FLaunchProfileTreeNode::FGetBool ForCode;
+		FLaunchProfileTreeNode::FGetBool ForDeployment;
+		FLaunchProfileTreeNode::FGetBool ForRun;
+		FLaunchProfileTreeNode::FGetString EmptyString;
+
+
+
+		TArray<FString> CachedMapsToCook;
+		enum class EMapOption : uint8
+		{
+			Startup,
+			Selected,
+		};
+		EMapOption MapOption = EMapOption::Startup;
+
+		TArray<FString> CachedDeployDeviceIDs;
+		enum class EDeployDeviceOption : uint8
+		{
+			Default,
+			Selected,
+		};
+		EDeployDeviceOption DeployDeviceOption = EDeployDeviceOption::Default;
+
+
+		FString GetProjectPath() const;
+		void SetProjectName(FString ProjectPath);
+		bool HasProject() const;
+
+		FString GetBuildTarget() const;
+		void SetBuildTarget(FString BuildTarget);
+		TArray<EBuildTargetType> GetSupportedBuildTargetTypes() const;
+
+		void SetBuildConfiguration(EBuildConfiguration BuildConfiguration);
+		EBuildConfiguration GetBuildConfiguration() const;
+
+		void SetContentScheme(EContentScheme ContentScheme);
+		bool IsContentSchemeAvailable(EContentScheme, FText& OutReason) const;
+
+		FString GetCommandLine() const;
+		void SetCommandLine( const FString& NewCommandLine );
+
+		void SetSelectedPlatforms( TArray<FString> SelectedPlatforms );
+		TArray<FString> GetSelectedPlatforms() const;
+
+		void SetCook( bool bCook );
+		bool GetCook( ILauncherProfilePtr InProfile = nullptr ) const;
+
+		void SetIncrementalCookMode( ELauncherProfileIncrementalCookMode::Type Mode );
+		ELauncherProfileIncrementalCookMode::Type GetIncrementalCookMode() const;
+
+		void SetMapsToCook(TArray<FString> MapsToCook);
+		TArray<FString> GetMapsToCook() const;
+		EMapOption GetMapOption() const;
+		void SetMapOption( EMapOption MapOption );
+		float GetMapListHeight() const;
+		void SetMapListHeight( float NewHeight );
+		TSharedRef<SWidget> CreateMapListWidget();
+		float MapListHeight = 300.0f;
+
+
+		void SetDeployDeviceIDs(TArray<FString> DeployDeviceIDs);
+		TArray<FString> GetDeployDeviceIDs() const;
+		EDeployDeviceOption GetDeployDeviceOption() const;
+		void SetDeployDeviceOption( EDeployDeviceOption DeployDeviceOption );
+		float GetDeployDeviceListHeight() const;
+		void SetDeployDeviceListHeight( float NewHeight );
+		void OnDeviceRemoved(FString DeviceID);
+		TSharedRef<SWidget> CreateDeployDeviceWidget();
+		float DeployDeviceListHeight = 150.0f;
+
+		void SetBuild( bool bBuild );
+		bool GetBuild( ILauncherProfilePtr InProfile = nullptr ) const;
+
+		void SetForceBuild( bool bForceBuild );
+		bool GetForceBuild( ILauncherProfilePtr InProfile = nullptr ) const;
+
+		void SetArchitecture( FString Architecture );
+		FString GetArchitecture() const;
+		FText GetArchitectureDisplayName( FString Architecture );
+
+		void SetDeployToDevice( bool bDeployToDevice );
+		bool GetDeployToDevice( ILauncherProfilePtr InProfile = nullptr ) const;
+
+		void SetIsRunning( bool bRun );
+		bool GetIsRunning( ILauncherProfilePtr InProfile = nullptr ) const;
+
+		void OnInitialMapChanged( TSharedPtr<FString> InitialMap, ESelectInfo::Type );
+		TSharedPtr<FString> GetInitialMap() const;
+
+
+
+	protected:
+
+		virtual void OnPropertyChanged() override;
+		void RefreshContentScheme();
+		void CacheStartupMapList();
+		void CacheArchitectures();
+		TSharedRef<SWidget> OnGenerateComboWidget( TSharedPtr<FString> InComboString );
+
+
+		FLaunchProfileTreeDataRef TreeData;
+		const ILauncherProfileRef Profile;
+		const ILauncherProfileRef DefaultProfile;
+		EProfileType ProfileType;
+
+		const TSharedRef<FModel> Model;
+
+		EContentScheme ContentScheme;
+		bool bShouldCook;
+		TArray<TSharedPtr<FString>> CachedStartupMaps;
+		TSharedPtr<SSearchableComboBox> InitalMapCombo;
+		TSharedPtr<SCustomLaunchMapListView> MapListView;
+		TSharedPtr<SCustomLaunchDeviceListView> DeployDeviceListView;
+		TArray<FString> CachedArchitectures;
+		EBuildTargetType CachedBuildTargetType = EBuildTargetType::Game;
+	};
+}
diff --git a/Engine/Source/Editor/TurnkeySupport/Private/TurnkeySupportModule.cpp b/Engine/Source/Editor/TurnkeySupport/Private/TurnkeySupportModule.cpp
--- a/Engine/Source/Editor/TurnkeySupport/Private/TurnkeySupportModule.cpp
+++ b/Engine/Source/Editor/TurnkeySupport/Private/TurnkeySupportModule.cpp
@@ -2372,7 +2372,7 @@ TSharedRef<SWidget> FTurnkeySupportModule::MakeTurnkeyMenuWidget() const
 		// options section
 		FToolMenuSection& OptionsSection = Menu->AddSection("TurnkeyOptions", LOCTEXT("TurnkeySection_Options", "Options and Settings"));
 		{
-			if (IPluginManager::Get().FindEnabledPlugin(TEXT("CustomLaunchUI")) != nullptr)
+			if (IPluginManager::Get().FindEnabledPlugin(TEXT("ProjectLauncher")) != nullptr)
 			{
 				OptionsSection.AddMenuEntry(
 					NAME_None,
diff --git a/Engine/Source/Programs/UnrealFrontend/UnrealFrontend.Target.cs b/Engine/Source/Programs/UnrealFrontend/UnrealFrontend.Target.cs
--- a/Engine/Source/Programs/UnrealFrontend/UnrealFrontend.Target.cs
+++ b/Engine/Source/Programs/UnrealFrontend/UnrealFrontend.Target.cs
@@ -31,7 +31,7 @@ public class UnrealFrontendTarget : TargetRules
 		// Old Profiler (SessionFrontend/Profiler) is deprecated since UE 5.0. Use Trace/UnrealInsights instead.
 		//GlobalDefinitions.Add("UE_DEPRECATED_PROFILER_ENABLED=1");
 
-		// new custom launch UI
-		AdditionalPlugins.Add("CustomLaunchUI");
+		// new project launcher
+		AdditionalPlugins.Add("ProjectLauncher");
 	}
 }
