updated Extensions menu for Project Launcher

#rnx

diff --git a/Engine/Plugins/Developer/ProjectLauncher/Source/ProjectLauncher/Private/Model/ProjectLauncherModel.cpp b/Engine/Plugins/Developer/ProjectLauncher/Source/ProjectLauncher/Private/Model/ProjectLauncherModel.cpp
--- a/Engine/Plugins/Developer/ProjectLauncher/Source/ProjectLauncher/Private/Model/ProjectLauncherModel.cpp
+++ b/Engine/Plugins/Developer/ProjectLauncher/Source/ProjectLauncher/Private/Model/ProjectLauncherModel.cpp
@@ -675,7 +675,20 @@ namespace ProjectLauncher
 		ApplyExtensionVariables(InProfile, InOutCommandLine, AsShared() );
 	}
 
+	bool FModel::AreExtensionsEnabled() const
+	{
+#if WITH_EDITOR
+		// make sure the extensions system is enabled
+		bool bEnableProjectLauncherExtensions = false;
+		GConfig->GetBool(TEXT("/Script/UnrealEd.EditorExperimentalSettings"), TEXT("bEnableProjectLauncherExtensions"), bEnableProjectLauncherExtensions, GEditorPerProjectIni);
+
+		return bEnableProjectLauncherExtensions;
+#else
 
+		// The ini file is not loaded by UnrealFrontend because its per-project, so allow extensions by default outside of the editor.
+		return true;
+#endif
+	}
 
 
 
diff --git a/Engine/Plugins/Developer/ProjectLauncher/Source/ProjectLauncher/Private/ProfileTree/LaunchProfileTreeData.cpp b/Engine/Plugins/Developer/ProjectLauncher/Source/ProjectLauncher/Private/ProfileTree/LaunchProfileTreeData.cpp
--- a/Engine/Plugins/Developer/ProjectLauncher/Source/ProjectLauncher/Private/ProfileTree/LaunchProfileTreeData.cpp
+++ b/Engine/Plugins/Developer/ProjectLauncher/Source/ProjectLauncher/Private/ProfileTree/LaunchProfileTreeData.cpp
@@ -280,19 +280,11 @@ namespace ProjectLauncher
 
 	FLaunchProfileTreeNode& FLaunchProfileTreeNode::AddCommandLineString( FText InName, FStringCallbacks&& StringCallbacks )
 	{
-		// make sure the extensions system is enabled
-		// ...unfortunately this setting is only available in the editor because the ini file is not loaded by UnrealFrontend because its per-project
-		bool bEnableProjectLauncherExtensions = false;
-		GConfig->GetBool(TEXT("/Script/UnrealEd.EditorExperimentalSettings"), TEXT("bEnableProjectLauncherExtensions"), bEnableProjectLauncherExtensions, GEditorPerProjectIni);
-
-		if (!bEnableProjectLauncherExtensions)
+		if (!TreeData->Model->AreExtensionsEnabled())
 		{
 			return AddString(InName, MoveTemp(StringCallbacks));
 		}
 
-
-
-
 		check(StringCallbacks.GetValue);
 		check(StringCallbacks.SetValue);
 
@@ -357,45 +349,6 @@ namespace ProjectLauncher
 			return MenuBuilder.MakeWidget();
 		};
 
-		auto OnGetExtensionsMenuContent = [this,StringCallbacks]()
-		{
-			const bool bShouldCloseWindowAfterMenuSelection = false;
-			const bool bCloseSelfOnly = false;
-			const bool bSearchable = false;
-			FMenuBuilder MenuBuilder(bShouldCloseWindowAfterMenuSelection, nullptr, nullptr, bCloseSelfOnly, &FCoreStyle::Get(), bSearchable );
-
-			for (TSharedPtr<FLaunchExtensionInstance> ExtensionInstance : TreeData->ExtensionInstances)
-			{
-				if (!ExtensionInstance->HasCustomExtensionMenu())
-				{
-					continue;
-				}
-
-#if 1
-				// flat list of extensions - easier to see but risks getting cluttered later on
-				MenuBuilder.BeginSection(NAME_None, ExtensionInstance->GetExtension()->GetDisplayName());
-				ExtensionInstance->MakeCustomExtensionSubmenu(MenuBuilder);
-				MenuBuilder.EndSection();
-#else 
-				// extensions in submenus - more scalable, but harder to see everything at a glance
-				const bool bInOpenSubMenuOnClick = false;
-				MenuBuilder.AddSubMenu(
-					ExtensionInstance->GetExtension()->GetDisplayName()),
-					FText::GetEmpty(),
-					FNewMenuDelegate::CreateSP( ExtensionInstance.ToSharedRef(), &FLaunchExtensionInstance::MakeExtensionSubmenu),
-					bInOpenSubMenuOnClick,
-					FSlateIcon(),
-					bShouldCloseWindowAfterMenuSelection
-				);
-#endif
-			}
-
-			return MenuBuilder.MakeWidget();
-		};
-
-
-
-
 		return AddWidget( InName,
 			MoveTemp(WidgetCallbacks),
 			SNew(SHorizontalBox)
@@ -433,27 +386,6 @@ namespace ProjectLauncher
 					.ColorAndOpacity(FSlateColor::UseForeground())
 				]
 			]
-
-			// command line extension button (@todo: temporary! this will be moved elsewhere)
-			+SHorizontalBox::Slot()
-			.AutoWidth()
-			.Padding(FMargin(0, 0, 3, 0))
-			.VAlign(VAlign_Center)
-			[
-				SNew(SComboButton)
-				.ComboButtonStyle(FAppStyle::Get(), "SimpleComboButton")
-				.ToolTipText(LOCTEXT("CmdLineExtensionsLabel", "Show additional command line extension options"))
-				.OnGetMenuContent_Lambda(OnGetExtensionsMenuContent)
-				.Visibility_Lambda( [this]() { return TreeData->bHasAnyMenuExtensions ? EVisibility::Visible : EVisibility::Collapsed; } )
-				.MenuPlacement(MenuPlacement_BelowRightAnchor)
-				.HasDownArrow(false)
-				.ButtonContent()
-				[
-					SNew(SImage)
-					.Image(FProjectLauncherStyle::Get().GetBrush("Icons.EllipsisVerticalNarrow"))
-					.ColorAndOpacity(FSlateColor::UseForeground())
-				]
-			]
 		);
 	}
 
diff --git a/Engine/Plugins/Developer/ProjectLauncher/Source/ProjectLauncher/Private/Styling/ProjectLauncherStyle.cpp b/Engine/Plugins/Developer/ProjectLauncher/Source/ProjectLauncher/Private/Styling/ProjectLauncherStyle.cpp
--- a/Engine/Plugins/Developer/ProjectLauncher/Source/ProjectLauncher/Private/Styling/ProjectLauncherStyle.cpp
+++ b/Engine/Plugins/Developer/ProjectLauncher/Source/ProjectLauncher/Private/Styling/ProjectLauncherStyle.cpp
@@ -73,7 +73,7 @@ void FProjectLauncherStyle::Initialize()
 
 	StyleSet->Set("Icons.DiffersFromDefault", new CORE_IMAGE_BRUSH_SVG( "../Editor/Slate/Starship/Common/ResetToDefault", Icon16x16) );
 	StyleSet->Set("Icons.DeviceManager", new CORE_IMAGE_BRUSH_SVG( "../Editor/Slate/Starship/Common/DeviceManager", Icon16x16) );
-	StyleSet->Set("Icons.EllipsisVerticalNarrow", new CORE_IMAGE_BRUSH_SVG("Starship/Common/ellipsis-vertical-narrow", FVector2f(6.f, 20.f))); // @todo: temporary - extension menu will move elsewhere
+	StyleSet->Set("Icons.Plugins", new CORE_IMAGE_BRUSH_SVG("../Editor/Slate/Starship/Common/Plugins", Icon16x16));
 
 	StyleSet->Set("Icons.Asset", new CORE_IMAGE_BRUSH( "Icons/doc_16x", Icon16x16 ) );
 
diff --git a/Engine/Plugins/Developer/ProjectLauncher/Source/ProjectLauncher/Private/Widgets/Profiles/SCustomLaunchCustomProfileEditor.cpp b/Engine/Plugins/Developer/ProjectLauncher/Source/ProjectLauncher/Private/Widgets/Profiles/SCustomLaunchCustomProfileEditor.cpp
--- a/Engine/Plugins/Developer/ProjectLauncher/Source/ProjectLauncher/Private/Widgets/Profiles/SCustomLaunchCustomProfileEditor.cpp
+++ b/Engine/Plugins/Developer/ProjectLauncher/Source/ProjectLauncher/Private/Widgets/Profiles/SCustomLaunchCustomProfileEditor.cpp
@@ -5,6 +5,7 @@
 #include "Styling/AppStyle.h"
 #include "Styling/StyleColors.h"
 #include "Styling/ProjectLauncherStyle.h"
+#include "Extension/LaunchExtension.h"
 #include "Framework/Commands/UIAction.h"
 #include "Framework/MultiBox/MultiBoxBuilder.h"
 #include "SlateOptMacros.h"
@@ -364,5 +365,55 @@ void SCustomLaunchCustomProfileEditor::SetProfile( const ILauncherProfilePtr& Pr
 	TreeView->RequestTreeRefresh();
 }
 
+bool SCustomLaunchCustomProfileEditor::HasAnyMenuExtensions() const
+{
+	if (TreeBuilder && TreeBuilder->GetProfileTree()->bHasAnyMenuExtensions)
+	{
+		return true;
+	}
+
+	return false;
+}
+
+TSharedRef<SWidget> SCustomLaunchCustomProfileEditor::MakeExtensionsMenu()
+{
+	ProjectLauncher::FLaunchProfileTreeDataRef TreeData = TreeBuilder->GetProfileTree();
+
+	const bool bShouldCloseWindowAfterMenuSelection = false;
+	const bool bCloseSelfOnly = false;
+	const bool bSearchable = false;
+	FMenuBuilder MenuBuilder(bShouldCloseWindowAfterMenuSelection, nullptr, nullptr, bCloseSelfOnly, &FCoreStyle::Get(), bSearchable );
+
+	for (TSharedPtr<ProjectLauncher::FLaunchExtensionInstance> ExtensionInstance : TreeData->ExtensionInstances)
+	{
+		if (!ExtensionInstance->HasCustomExtensionMenu())
+		{
+			continue;
+		}
+
+#if 1
+		// flat list of extensions - easier to see but risks getting cluttered later on
+		MenuBuilder.BeginSection(NAME_None, ExtensionInstance->GetExtension()->GetDisplayName());
+		ExtensionInstance->MakeCustomExtensionSubmenu(MenuBuilder);
+		MenuBuilder.EndSection();
+#else 
+		// extensions in submenus - more scalable, but harder to see everything at a glance
+		const bool bInOpenSubMenuOnClick = false;
+		MenuBuilder.AddSubMenu(
+			ExtensionInstance->GetExtension()->GetDisplayName()),
+			FText::GetEmpty(),
+			FNewMenuDelegate::CreateSP( ExtensionInstance.ToSharedRef(), &FLaunchExtensionInstance::MakeExtensionSubmenu),
+			bInOpenSubMenuOnClick,
+			FSlateIcon(),
+			bShouldCloseWindowAfterMenuSelection
+		);
+#endif
+	}
+
+	return MenuBuilder.MakeWidget();
+}
+
+
+
 
 #undef LOCTEXT_NAMESPACE
diff --git a/Engine/Plugins/Developer/ProjectLauncher/Source/ProjectLauncher/Private/Widgets/Profiles/SCustomLaunchCustomProfileEditor.h b/Engine/Plugins/Developer/ProjectLauncher/Source/ProjectLauncher/Private/Widgets/Profiles/SCustomLaunchCustomProfileEditor.h
--- a/Engine/Plugins/Developer/ProjectLauncher/Source/ProjectLauncher/Private/Widgets/Profiles/SCustomLaunchCustomProfileEditor.h
+++ b/Engine/Plugins/Developer/ProjectLauncher/Source/ProjectLauncher/Private/Widgets/Profiles/SCustomLaunchCustomProfileEditor.h
@@ -32,6 +32,8 @@ public:
 	void SetProfile( const ILauncherProfilePtr& Profile );
 	ILauncherProfilePtr GetProfile() const { return CurrentProfile; }
 
+	bool HasAnyMenuExtensions() const;
+	TSharedRef<SWidget> MakeExtensionsMenu();
 
 	// IScrollableWidget interface
 	virtual FVector2D GetScrollDistance() override;
diff --git a/Engine/Plugins/Developer/ProjectLauncher/Source/ProjectLauncher/Private/Widgets/Profiles/SCustomLaunchProfilesPanel.cpp b/Engine/Plugins/Developer/ProjectLauncher/Source/ProjectLauncher/Private/Widgets/Profiles/SCustomLaunchProfilesPanel.cpp
--- a/Engine/Plugins/Developer/ProjectLauncher/Source/ProjectLauncher/Private/Widgets/Profiles/SCustomLaunchProfilesPanel.cpp
+++ b/Engine/Plugins/Developer/ProjectLauncher/Source/ProjectLauncher/Private/Widgets/Profiles/SCustomLaunchProfilesPanel.cpp
@@ -25,6 +25,7 @@
 #include "ITargetDeviceProxy.h"
 #include "Modules/ModuleManager.h"
 #include "Framework/Docking/TabManager.h"
+#include "Framework/MultiBox/MultiBoxBuilder.h"
 #include "PlatformInfo.h"
 
 
@@ -332,32 +333,53 @@ TSharedRef<SWidget> SCustomLaunchProfilesPanel::CreateProfileEditorWidget()
 		]
 	]
 
-	// hide properties buttom button
+	// extensions button bar
 	+SVerticalBox::Slot()
 	.AutoHeight()
 	[
 		SNew(SBorder)
 		.BorderImage(FAppStyle::GetBrush("Brushes.Header"))
-		.Visibility_Lambda( [this]() { return IsProfileEditorVisible() ? EVisibility::Visible : EVisibility::Collapsed; } )
+		.Visibility(this, &SCustomLaunchProfilesPanel::GetExtensionsBarVisibility)
 		[
 			SNew(SHorizontalBox)
 
 			+SHorizontalBox::Slot()
-			.FillWidth(1)
+			.AutoWidth()
+			.VAlign(VAlign_Center)
+			.HAlign(HAlign_Left)
+			.Padding(4,4)
 			[
-				SNew(SSpacer)
+				SNew(SComboButton)
+				.ComboButtonStyle(FAppStyle::Get(), "SimpleComboButton")
+				.ToolTipText(LOCTEXT("ExtensionsToolTip", "Select launcher extensions."))
+				.OnGetMenuContent(PropertyEditor.ToSharedRef(), &SCustomLaunchCustomProfileEditor::MakeExtensionsMenu)
+				.MenuPlacement(MenuPlacement_ComboBoxRight)
+				.ButtonContent()
+				[
+					SNew(SHorizontalBox)
+					+SHorizontalBox::Slot()
+					.AutoWidth()
+					.Padding(2, 0)
+					[
+						SNew(SImage)
+						.Image(FProjectLauncherStyle::Get().GetBrush("Icons.Plugins"))
+						.ColorAndOpacity(FSlateColor::UseForeground())
+					]
+					+SHorizontalBox::Slot()
+					.AutoWidth()
+					.Padding(2, 0)
+					[
+						SNew(STextBlock)
+						.Text(LOCTEXT("ExtensionsLabel", "Extensions"))
+						.ColorAndOpacity(FSlateColor::UseForeground())
+					]
+				]
 			]
 
 			+SHorizontalBox::Slot()
-			.AutoWidth()
-			.VAlign(VAlign_Center)
-			.HAlign(HAlign_Right)
-			.Padding(4,4)
+			.FillWidth(1)
 			[
-				SNew(SButton)
-				.OnClicked_Lambda( [this]() { SetProfileEditorVisible(false); return FReply::Handled();} )
-				.ButtonStyle(FAppStyle::Get(), "Button")
-				.Text(LOCTEXT("HidePropertiesLabel", "Hide Properties"))
+				SNew(SSpacer)
 			]
 		]
 	]
@@ -520,6 +542,18 @@ TSharedRef<SWidget> SCustomLaunchProfilesPanel::CreateAdvancedProfileWarningWidg
 END_SLATE_FUNCTION_BUILD_OPTIMIZATION
 
 
+EVisibility SCustomLaunchProfilesPanel::GetExtensionsBarVisibility() const
+{
+	if (IsProfileEditorVisible() && Model->AreExtensionsEnabled() && PropertyEditor->HasAnyMenuExtensions())
+	{
+		return EVisibility::Visible;
+	}
+	else
+	{
+		return EVisibility::Collapsed; 
+	}
+}
+
 
 void SCustomLaunchProfilesPanel::OnProfileEdit(const ILauncherProfilePtr& Profile)
 {
diff --git a/Engine/Plugins/Developer/ProjectLauncher/Source/ProjectLauncher/Private/Widgets/Profiles/SCustomLaunchProfilesPanel.h b/Engine/Plugins/Developer/ProjectLauncher/Source/ProjectLauncher/Private/Widgets/Profiles/SCustomLaunchProfilesPanel.h
--- a/Engine/Plugins/Developer/ProjectLauncher/Source/ProjectLauncher/Private/Widgets/Profiles/SCustomLaunchProfilesPanel.h
+++ b/Engine/Plugins/Developer/ProjectLauncher/Source/ProjectLauncher/Private/Widgets/Profiles/SCustomLaunchProfilesPanel.h
@@ -36,6 +36,7 @@ private:
 	TSharedRef<SWidget> CreateAdvancedProfileWarningWidget();
 	TSharedRef<SWidget> CreateProfileEditorToolbarWidget();
 
+	EVisibility GetExtensionsBarVisibility() const;
 
 	void OnProfileEdit(const ILauncherProfilePtr& Profile);
 	void OnProfileDelete(const ILauncherProfilePtr& Profile);
@@ -75,7 +76,6 @@ private:
 	void SetDefaultBuildTarget(FString BuildTarget);
 	FString GetDefaultBuildTarget() const;
 
-
 	ProjectLauncher::FOnProfileClicked OnProfileLaunchClicked;
 	TSharedPtr<ProjectLauncher::FModel> Model;
 
diff --git a/Engine/Plugins/Developer/ProjectLauncher/Source/ProjectLauncher/Public/Model/ProjectLauncherModel.h b/Engine/Plugins/Developer/ProjectLauncher/Source/ProjectLauncher/Public/Model/ProjectLauncherModel.h
--- a/Engine/Plugins/Developer/ProjectLauncher/Source/ProjectLauncher/Public/Model/ProjectLauncherModel.h
+++ b/Engine/Plugins/Developer/ProjectLauncher/Source/ProjectLauncher/Public/Model/ProjectLauncherModel.h
@@ -171,6 +171,8 @@ namespace ProjectLauncher
 		const TCHAR* GetConfigSection() const;
 		const FString& GetConfigIni() const;
 
+		bool AreExtensionsEnabled() const;
+
 	private:
 		FString ConfigFileName;
 
