Improve map enumeration:
- new fast path for the current in-editor project, using the asset registry instead.
- map list is cached for other projects
- on-screen 'slow task' dialog during disk search
- lazy map list construction: map tree only populated when first visible + map combo only populated when drop-down first opens.

#rnx

diff --git a/Engine/Plugins/Developer/ProjectLauncher/Source/ProjectLauncher/Private/Model/ProjectLauncherModel.cpp b/Engine/Plugins/Developer/ProjectLauncher/Source/ProjectLauncher/Private/Model/ProjectLauncherModel.cpp
--- a/Engine/Plugins/Developer/ProjectLauncher/Source/ProjectLauncher/Private/Model/ProjectLauncherModel.cpp
+++ b/Engine/Plugins/Developer/ProjectLauncher/Source/ProjectLauncher/Private/Model/ProjectLauncherModel.cpp
@@ -5,6 +5,7 @@
 #include "Misc/ConfigCacheIni.h"
 #include "Misc/ConfigContext.h"
 #include "Misc/App.h"
+#include "Misc/ScopedSlowTask.h"
 #include "PlatformInfo.h"
 #include "GameProjectHelper.h"
 #include "ITargetDeviceProxy.h"
@@ -16,6 +17,12 @@
 #include "Settings/ProjectPackagingSettings.h"
 #include "Experimental/ZenServerInterface.h"
 
+#if WITH_EDITOR
+#include "Engine/World.h"
+#include "AssetRegistry/AssetRegistryModule.h"
+#include "AssetRegistry/AssetData.h"
+#endif
+
 #define LOCTEXT_NAMESPACE "SCustomLaunchProfileSelector"
 
 namespace ProjectLauncher
@@ -691,6 +698,111 @@ namespace ProjectLauncher
 	}
 
 
+	TArray<FString> FModel::GetAndCacheMapPaths(const FString& InOptionalProjectPath) // @todo: map list parsing should ideally be asyncronous, showing a spinner in the map selector controls until its finished etc.
+	{
+		// prepare values
+		FString ProjectPath;
+		FString ProjectName;
+		FString ContentDir;
+		bool bWantEngineMaps = InOptionalProjectPath.IsEmpty();
+		if (bWantEngineMaps)
+		{
+			ContentDir = FPaths::Combine(FPaths::EngineContentDir(), TEXT("Maps"));
+		}
+		else
+		{
+			ProjectPath = FPaths::IsRelative(InOptionalProjectPath) ? FPaths::Combine(FPaths::RootDir(), InOptionalProjectPath) : InOptionalProjectPath;
+			ProjectName = FPaths::GetBaseFilename(ProjectPath);
+			ContentDir = FPaths::Combine(ProjectPath, TEXT("Content")); // @todo: do we want to show plugin maps as well?
+		}
+				
+		// the editor has access to the asset registry which is must faster and will be up to date if new maps are added at runtime - if this is the project that's selected in the editor
+#if WITH_EDITOR
+		if (FPaths::IsProjectFilePathSet())
+		{
+			FString GlobalProjectFilePath = FPaths::ConvertRelativePathToFull(FPaths::GetProjectFilePath());
+			FString GlobalProjectPath = FPaths::GetPath(GlobalProjectFilePath);
+
+			if (!GlobalProjectPath.IsEmpty() && (GlobalProjectPath == ProjectPath || bWantEngineMaps) )
+			{
+				// gather all world asset metadata
+				FAssetRegistryModule& AssetRegistryModule = FModuleManager::LoadModuleChecked<FAssetRegistryModule>("AssetRegistry");
+				TArray<FAssetData> MapAssets;
+				AssetRegistryModule.Get().GetAssetsByClass(UWorld::StaticClass()->GetClassPathName(), MapAssets, true);
+
+				// build the list of map file names for the maps in the desired content directory
+				TArray<FString> MapFileList;
+				for (const FAssetData& MapAsset : MapAssets)
+				{
+					FString MapFileName = FPackageName::LongPackageNameToFilename(MapAsset.PackageName.ToString(), FPackageName::GetMapPackageExtension());
+					MapFileName = FPaths::ConvertRelativePathToFull(MapFileName);
+
+					if (FPaths::IsUnderDirectory(MapFileName, ContentDir))
+					{
+						MapFileList.Add( MapFileName );
+					}
+				}
+
+				MapFileList.Sort();
+				return MapFileList;
+			}
+		}
+#endif
+
+		// check to see if we've cached this project's maps already
+		const TArray<FString>* CachedProjectMapListPtr = CachedMapPaths.Find(ProjectName);
+		if (CachedProjectMapListPtr != nullptr)
+		{
+			return *CachedProjectMapListPtr;
+		}
+		else
+		{
+			// this is slow & currently blocking so show a wait dialog. 
+			// annoyingly this doesn't show up in UnrealFrontend because that uses FFeedbackContext rather than FFeedbackContextEditor
+			FScopedSlowTask SlowTask(0, LOCTEXT("CacheProjectMapsDesc","Caching project maps"));
+			SlowTask.MakeDialog();
+
+			// search for map files
+			TArray<FString> MapFileList;
+			const FString WildCard = FString::Printf(TEXT("*%s"), *FPackageName::GetMapPackageExtension());
+			IFileManager::Get().FindFilesRecursive(MapFileList, *ContentDir, *WildCard, true, false);
+
+			MapFileList.Sort();
+			CachedMapPaths.Add(ProjectName, MapFileList);
+			return MapFileList;
+		}
+	}
+
+	TArray<FString> FModel::GetAvailableProjectMapNames( const FString& InProjectPath )
+	{
+		TArray<FString> Maps = GetAndCacheMapPaths(InProjectPath);
+		for (FString& Map : Maps)
+		{
+			Map = FPaths::GetBaseFilename(Map);
+		}
+		return MoveTemp(Maps);
+	}
+
+	TArray<FString> FModel::GetAvailableProjectMapPaths( const FString& InProjectPath )
+	{
+		return GetAndCacheMapPaths(InProjectPath);
+	}
+
+	TArray<FString> FModel::GetAvailableEngineMapNames()
+	{
+		TArray<FString> Maps = GetAndCacheMapPaths(FString());
+		for (FString& Map : Maps)
+		{
+			Map = FPaths::GetBaseFilename(Map);
+		}
+		return MoveTemp(Maps);
+
+	}
+
+	TArray<FString> FModel::GetAvailableEngineMapPaths()
+	{
+		return GetAndCacheMapPaths(FString());
+	}
 
 
 
diff --git a/Engine/Plugins/Developer/ProjectLauncher/Source/ProjectLauncher/Private/ProfileTree/GenericProfileTreeBuilder.cpp b/Engine/Plugins/Developer/ProjectLauncher/Source/ProjectLauncher/Private/ProfileTree/GenericProfileTreeBuilder.cpp
--- a/Engine/Plugins/Developer/ProjectLauncher/Source/ProjectLauncher/Private/ProfileTree/GenericProfileTreeBuilder.cpp
+++ b/Engine/Plugins/Developer/ProjectLauncher/Source/ProjectLauncher/Private/ProfileTree/GenericProfileTreeBuilder.cpp
@@ -92,7 +92,6 @@ namespace ProjectLauncher
 		ContentScheme = Model->DetermineProfileContentScheme(Profile);
 		CachedBuildTargetType = Model->GetBuildTargetInfo(Profile).Type;
 	
-		CacheStartupMapList();
 		CacheArchitectures();
 	}
 
@@ -500,10 +499,11 @@ namespace ProjectLauncher
 			.OptionsSource(&CachedStartupMaps)
 			.OnSelectionChanged( this, &FGenericProfileTreeBuilder::OnInitialMapChanged )
 			.OnGenerateWidget( this, &FGenericProfileTreeBuilder::OnGenerateComboWidget )
+			.OnComboBoxOpening( this, &FGenericProfileTreeBuilder::CacheStartupMapList )
 			.Content()
 			[
 				SNew(STextBlock)
-				.Text_Lambda( [this]() { return FText::FromString( *GetInitialMap() ); } )
+				.Text_Lambda( [this]() { return FText::FromString( *Profile->GetDefaultLaunchRole()->GetInitialMap() ); } )
 				.Font(FCoreStyle::Get().GetFontStyle("SmallFont"))
 			]
 		);
@@ -557,11 +557,17 @@ namespace ProjectLauncher
 	}
 
 
-	void FGenericProfileTreeBuilder::CacheStartupMapList()
+	void FGenericProfileTreeBuilder::CacheStartupMapList() const
 	{
+		if (!bStartupMapCacheDirty)
+		{
+			return;
+		}
+		bStartupMapCacheDirty = false;
+
 		CachedStartupMaps.Reset();
 		CachedStartupMaps.Add(MakeShared<FString>());
-		for ( const FString& Map : FGameProjectHelper::GetAvailableMaps(Profile->GetProjectBasePath(), Profile->SupportsEngineMaps(), true))
+		for ( const FString& Map : Model->GetAvailableProjectMapNames(Profile->GetProjectBasePath()))
 		{
 			CachedStartupMaps.Add(MakeShared<FString>(Map));
 		}
@@ -692,7 +698,7 @@ namespace ProjectLauncher
 			MapListView->RefreshMapList();
 		}
 
-		CacheStartupMapList();
+		bStartupMapCacheDirty = true;
 	}
 
 	bool FGenericProfileTreeBuilder::HasProject() const
@@ -858,7 +864,7 @@ namespace ProjectLauncher
 	TSharedRef<SWidget> FGenericProfileTreeBuilder::CreateMapListWidget()
 	{
 		MapListView = 
-			SNew(SCustomLaunchMapListView)
+			SNew(SCustomLaunchMapListView, Model)
 			.OnSelectionChanged(this, &FGenericProfileTreeBuilder::SetMapsToCook)
 			.SelectedMaps(this, &FGenericProfileTreeBuilder::GetMapsToCook)
 			.ProjectPath(this, &FGenericProfileTreeBuilder::GetProjectPath)
@@ -1341,6 +1347,8 @@ namespace ProjectLauncher
 
 	TSharedPtr<FString> FGenericProfileTreeBuilder::GetInitialMap() const
 	{
+		CacheStartupMapList();
+
 		FString InitialMap = Profile->GetDefaultLaunchRole()->GetInitialMap();
 		const TSharedPtr<FString>* FoundMapPtr = CachedStartupMaps.FindByPredicate( [InitialMap]( const TSharedPtr<FString>& Map )
 		{
diff --git a/Engine/Plugins/Developer/ProjectLauncher/Source/ProjectLauncher/Private/Widgets/Shared/SCustomLaunchMapListView.cpp b/Engine/Plugins/Developer/ProjectLauncher/Source/ProjectLauncher/Private/Widgets/Shared/SCustomLaunchMapListView.cpp
--- a/Engine/Plugins/Developer/ProjectLauncher/Source/ProjectLauncher/Private/Widgets/Shared/SCustomLaunchMapListView.cpp
+++ b/Engine/Plugins/Developer/ProjectLauncher/Source/ProjectLauncher/Private/Widgets/Shared/SCustomLaunchMapListView.cpp
@@ -21,8 +21,9 @@
 
 
 BEGIN_SLATE_FUNCTION_BUILD_OPTIMIZATION
-void SCustomLaunchMapListView::Construct(const FArguments& InArgs)
+void SCustomLaunchMapListView::Construct(const FArguments& InArgs, TSharedRef<ProjectLauncher::FModel> InModel)
 {
+	Model = InModel.ToSharedPtr();
 	OnSelectionChanged = InArgs._OnSelectionChanged;
 	SelectedMaps = InArgs._SelectedMaps;
 	ProjectPath = InArgs._ProjectPath;
@@ -209,6 +210,13 @@ void SCustomLaunchMapListView::OnProjectChanged()
 
 void SCustomLaunchMapListView::RefreshMapList()
 {
+	bMapListDirty = true;
+}
+
+void SCustomLaunchMapListView::RefreshMapListInternal()
+{
+	bMapListDirty = false;
+
 	// helper to build the node tree
 	MapTreeRoot = MakeShared<FMapTreeNode>();
 	auto AddMapPath = [this]( const FString& Path )
@@ -265,8 +273,7 @@ void SCustomLaunchMapListView::RefreshMapList()
 			ProjectContentDir = FPaths::Combine(FPaths::RootDir(), ProjectContentDir);
 		}
 
-		TArray<FString> ProjectMaps;
-		IFileManager::Get().FindFilesRecursive(ProjectMaps, *ProjectContentDir, *WildCard, true, false);
+		TArray<FString> ProjectMaps = Model->GetAvailableProjectMapPaths(ProjectBasePath);
 		for (FString& ProjectMap : ProjectMaps)
 		{
 			AvailableMaps.AddUnique(FPaths::GetBaseFilename(ProjectMap));
@@ -286,8 +293,7 @@ void SCustomLaunchMapListView::RefreshMapList()
 	if (bShowEngineMaps)
 	{
 		FString EngineMapDir = FPaths::Combine(FPaths::EngineContentDir(), TEXT("Maps"));
-		TArray<FString> EngineMaps;
-		IFileManager::Get().FindFilesRecursive(EngineMaps, *EngineMapDir, *WildCard, true, false);
+		TArray<FString> EngineMaps = Model->GetAvailableEngineMapPaths();
 		for (FString& EngineMap : EngineMaps)
 		{
 			AvailableMaps.AddUnique(FPaths::GetBaseFilename(EngineMap));
@@ -396,5 +402,24 @@ void SCustomLaunchMapListView::SetCheckBoxStateRecursive(FMapTreeNodePtr Node, E
 }
 
 
+void SCustomLaunchMapListView::Tick(const FGeometry& AllottedGeometry, const double InCurrentTime, const float InDeltaTime)
+{
+	if (bMapListDirty && bHasPaintedThisFrame)
+	{
+		RefreshMapListInternal();
+	}
+
+	bHasPaintedThisFrame = false;
+	SCompoundWidget::Tick(AllottedGeometry, InCurrentTime, InDeltaTime);
+}
+
+int32 SCustomLaunchMapListView::OnPaint( const FPaintArgs& Args, const FGeometry& AllottedGeometry, const FSlateRect& MyCullingRect, FSlateWindowElementList& OutDrawElements, int32 LayerId, const FWidgetStyle& InWidgetStyle, bool bParentEnabled) const
+{
+	bHasPaintedThisFrame = true;
+    return SCompoundWidget::OnPaint(Args, AllottedGeometry, MyCullingRect, OutDrawElements, LayerId, InWidgetStyle, bParentEnabled);
+}
+
+
+
 
 #undef LOCTEXT_NAMESPACE
diff --git a/Engine/Plugins/Developer/ProjectLauncher/Source/ProjectLauncher/ProjectLauncher.Build.cs b/Engine/Plugins/Developer/ProjectLauncher/Source/ProjectLauncher/ProjectLauncher.Build.cs
--- a/Engine/Plugins/Developer/ProjectLauncher/Source/ProjectLauncher/ProjectLauncher.Build.cs
+++ b/Engine/Plugins/Developer/ProjectLauncher/Source/ProjectLauncher/ProjectLauncher.Build.cs
@@ -28,5 +28,14 @@ public class ProjectLauncher : ModuleRules
 			"SlateCore",
 			"LauncherServices",
 		]);
+
+		if (Target.bCompileAgainstEditor)
+		{
+			PrivateDependencyModuleNames.AddRange(
+			[
+				"Engine",
+				"AssetRegistry",
+			]);
+		}
 	}
 }
diff --git a/Engine/Plugins/Developer/ProjectLauncher/Source/ProjectLauncher/Public/Model/ProjectLauncherModel.h b/Engine/Plugins/Developer/ProjectLauncher/Source/ProjectLauncher/Public/Model/ProjectLauncherModel.h
--- a/Engine/Plugins/Developer/ProjectLauncher/Source/ProjectLauncher/Public/Model/ProjectLauncherModel.h
+++ b/Engine/Plugins/Developer/ProjectLauncher/Source/ProjectLauncher/Public/Model/ProjectLauncherModel.h
@@ -173,6 +173,11 @@ namespace ProjectLauncher
 
 		bool AreExtensionsEnabled() const;
 
+		TArray<FString> GetAvailableProjectMapNames( const FString& InProjectPath );
+		TArray<FString> GetAvailableProjectMapPaths( const FString& InProjectPath );
+		TArray<FString> GetAvailableEngineMapNames();
+		TArray<FString> GetAvailableEngineMapPaths();
+
 	private:
 		FString ConfigFileName;
 
@@ -188,6 +193,8 @@ namespace ProjectLauncher
 
 		void OnModifyLaunchCommandLine( const ILauncherProfileRef& InProfile, FString& InOutCommandLine );
 
+		TArray<FString> GetAndCacheMapPaths(const FString& InOptionalProjectPath);
+
 		TSharedRef<ITargetDeviceProxyManager> DeviceProxyManager;
 		TSharedRef<ILauncher> Launcher;
 		TSharedRef<ILauncherProfileManager> ProfileManager;
@@ -201,6 +208,7 @@ namespace ProjectLauncher
 		TArray<TSharedPtr<FLaunchLogMessage>> LaunchLogMessages;
 
 		TMap<FString,FProjectSettings> CachedProjectSettings;
+		TMap<FString,TArray<FString>> CachedMapPaths;
 
 		FOnSelectedProfileChanged ProfileSelectedDelegate;
 		FOnSelectedPlatformChanged PlatformChangedDelegate;
diff --git a/Engine/Plugins/Developer/ProjectLauncher/Source/ProjectLauncher/Public/ProfileTree/GenericProfileTreeBuilder.h b/Engine/Plugins/Developer/ProjectLauncher/Source/ProjectLauncher/Public/ProfileTree/GenericProfileTreeBuilder.h
--- a/Engine/Plugins/Developer/ProjectLauncher/Source/ProjectLauncher/Public/ProfileTree/GenericProfileTreeBuilder.h
+++ b/Engine/Plugins/Developer/ProjectLauncher/Source/ProjectLauncher/Public/ProfileTree/GenericProfileTreeBuilder.h
@@ -171,7 +171,7 @@ namespace ProjectLauncher
 
 		virtual void OnPropertyChanged() override;
 		void RefreshContentScheme();
-		void CacheStartupMapList();
+		void CacheStartupMapList() const;
 		void CacheArchitectures();
 		TSharedRef<SWidget> OnGenerateComboWidget( TSharedPtr<FString> InComboString );
 
@@ -185,7 +185,8 @@ namespace ProjectLauncher
 
 		EContentScheme ContentScheme;
 		bool bShouldCook;
-		TArray<TSharedPtr<FString>> CachedStartupMaps;
+		mutable bool bStartupMapCacheDirty = true;
+		mutable TArray<TSharedPtr<FString>> CachedStartupMaps;
 		TSharedPtr<SSearchableComboBox> InitalMapCombo;
 		TSharedPtr<SCustomLaunchMapListView> MapListView;
 		TSharedPtr<SCustomLaunchDeviceListView> DeployDeviceListView;
diff --git a/Engine/Plugins/Developer/ProjectLauncher/Source/ProjectLauncher/Public/Widgets/Shared/SCustomLaunchMapListView.h b/Engine/Plugins/Developer/ProjectLauncher/Source/ProjectLauncher/Public/Widgets/Shared/SCustomLaunchMapListView.h
--- a/Engine/Plugins/Developer/ProjectLauncher/Source/ProjectLauncher/Public/Widgets/Shared/SCustomLaunchMapListView.h
+++ b/Engine/Plugins/Developer/ProjectLauncher/Source/ProjectLauncher/Public/Widgets/Shared/SCustomLaunchMapListView.h
@@ -4,6 +4,7 @@
 
 #include "CoreMinimal.h"
 #include "Widgets/SCompoundWidget.h"
+#include "Model/ProjectLauncherModel.h"
 
 class ITableRow;
 class STableViewBase;
@@ -22,12 +23,13 @@ public:
 	SLATE_END_ARGS()
 
 public:
-	void Construct(	const FArguments& InArgs );
+	void Construct( const FArguments& InArgs, TSharedRef<ProjectLauncher::FModel> InModel );
 
 	void RefreshMapList();
 	TSharedRef<SWidget> MakeControlsWidget();
 
 protected:
+	TSharedPtr<ProjectLauncher::FModel> Model;
 	TAttribute<TArray<FString>> SelectedMaps;
 	TAttribute<FString> ProjectPath;
 	FOnSelectionChanged OnSelectionChanged;
@@ -67,4 +69,11 @@ private:
 	void SetMapTreeNodeCheckState(ECheckBoxState CheckBoxState, FMapTreeNodePtr Item);
 	const FSlateBrush* GetMapTreeNodeIcon(FMapTreeNodePtr Node) const;
 	FSlateColor GetMapTreeNodeColor(FMapTreeNodePtr Node) const;
+
+	mutable bool bHasPaintedThisFrame = false;
+	virtual int32 OnPaint( const FPaintArgs& Args, const FGeometry& AllottedGeometry, const FSlateRect& MyCullingRect, FSlateWindowElementList& OutDrawElements, int32 LayerId, const FWidgetStyle& InWidgetStyle, bool bParentEnabled) const override;
+
+	bool bMapListDirty = false;
+	void RefreshMapListInternal();
+	virtual void Tick(const FGeometry& AllottedGeometry, const double InCurrentTime, const float InDeltaTime) override;
 };
