filter Fatal Errors correctly in Project Launcher output

#rnx
#rnx

diff --git a/Engine/Plugins/Developer/ProjectLauncher/Source/ProjectLauncher/Private/Widgets/Output/SCustomLaunchLaunchPanel.cpp b/Engine/Plugins/Developer/ProjectLauncher/Source/ProjectLauncher/Private/Widgets/Output/SCustomLaunchLaunchPanel.cpp
--- a/Engine/Plugins/Developer/ProjectLauncher/Source/ProjectLauncher/Private/Widgets/Output/SCustomLaunchLaunchPanel.cpp
+++ b/Engine/Plugins/Developer/ProjectLauncher/Source/ProjectLauncher/Private/Widgets/Output/SCustomLaunchLaunchPanel.cpp
@@ -387,6 +387,10 @@ void SCustomLaunchLaunchPanel::HandleOutputReceived(const FString& InMessage)
 	{
 		Verbosity = ELogVerbosity::Display;
 	}
+	else if ( InMessage.Contains(TEXT("Fatal error"), ESearchCase::IgnoreCase) )
+	{
+		Verbosity = ELogVerbosity::Fatal;
+	}
 	else if ( InMessage.Contains(TEXT("Error:"), ESearchCase::IgnoreCase) )
 	{
 		Verbosity = ELogVerbosity::Error;
