Reset incompatible options when changing content scheme & rerun verification

#rnx

diff --git a/Engine/Plugins/Developer/ProjectLauncher/Source/ProjectLauncher/Private/Model/ProjectLauncherModel.cpp b/Engine/Plugins/Developer/ProjectLauncher/Source/ProjectLauncher/Private/Model/ProjectLauncherModel.cpp
--- a/Engine/Plugins/Developer/ProjectLauncher/Source/ProjectLauncher/Private/Model/ProjectLauncherModel.cpp
+++ b/Engine/Plugins/Developer/ProjectLauncher/Source/ProjectLauncher/Private/Model/ProjectLauncherModel.cpp
@@ -541,6 +541,11 @@ namespace ProjectLauncher
 		Profile->SetUseZenPakStreaming(bZenPakStreaming);
 		Profile->SetDeployWithUnrealPak(bPakFiles);
 		Profile->SetUseZenStore(bUseZen);
+		if (!bPakFiles)
+		{
+			Profile->SetGenerateChunks(false);
+			Profile->SetUseIoStore(false);
+		}
 
 		ELauncherProfileDeploymentModes::Type DeploymentMode = DefaultDeploymentMode;
 		if (bCOTF)
diff --git a/Engine/Source/Developer/LauncherServices/Private/Profiles/LauncherProfile.h b/Engine/Source/Developer/LauncherServices/Private/Profiles/LauncherProfile.h
--- a/Engine/Source/Developer/LauncherServices/Private/Profiles/LauncherProfile.h
+++ b/Engine/Source/Developer/LauncherServices/Private/Profiles/LauncherProfile.h
@@ -2736,6 +2736,8 @@ public:
 		{
 			SetDeployWithUnrealPak(true);
 		}
+
+		Validate();
 	}
 
 	virtual bool IsUsingIoStore() const override
@@ -2746,6 +2748,8 @@ public:
 	virtual void SetUseZenStore(bool bInUseZenStore) override
 	{
 		bUseZenStore = bInUseZenStore;
+
+		Validate();
 	}
 
 	virtual bool IsUsingZenStore() const override
