Launch profile tree view refreshes properly

#rnx

diff --git a/Engine/Plugins/Developer/ProjectLauncher/Source/ProjectLauncher/Private/ProfileTree/GenericProfileTreeBuilder.cpp b/Engine/Plugins/Developer/ProjectLauncher/Source/ProjectLauncher/Private/ProfileTree/GenericProfileTreeBuilder.cpp
--- a/Engine/Plugins/Developer/ProjectLauncher/Source/ProjectLauncher/Private/ProfileTree/GenericProfileTreeBuilder.cpp
+++ b/Engine/Plugins/Developer/ProjectLauncher/Source/ProjectLauncher/Private/ProfileTree/GenericProfileTreeBuilder.cpp
@@ -554,6 +554,8 @@ namespace ProjectLauncher
 		{
 			Model->GetProfileManager()->SaveJSONProfile(Profile);
 		}
+
+		TreeData->RequestTreeRefresh();
 	}
 
 
@@ -997,6 +999,8 @@ namespace ProjectLauncher
 		static const float MinMapListHeight = 100.0f;
 
 		MapListHeight = FMath::Max(NewHeight, MinMapListHeight);
+
+		TreeData->RequestTreeRefresh();
 	}
 
 	FGenericProfileTreeBuilder::EMapOption FGenericProfileTreeBuilder::GetMapOption() const
@@ -1146,6 +1150,8 @@ namespace ProjectLauncher
 		static const float MinDeployDeviceListHeight = 100.0f;
 
 		DeployDeviceListHeight = FMath::Max(NewHeight, MinDeployDeviceListHeight);
+
+		TreeData->RequestTreeRefresh();
 	}
 
 	void FGenericProfileTreeBuilder::OnDeviceRemoved(FString DeviceID)
diff --git a/Engine/Plugins/Developer/ProjectLauncher/Source/ProjectLauncher/Private/ProfileTree/LaunchProfileTreeData.cpp b/Engine/Plugins/Developer/ProjectLauncher/Source/ProjectLauncher/Private/ProfileTree/LaunchProfileTreeData.cpp
--- a/Engine/Plugins/Developer/ProjectLauncher/Source/ProjectLauncher/Private/ProfileTree/LaunchProfileTreeData.cpp
+++ b/Engine/Plugins/Developer/ProjectLauncher/Source/ProjectLauncher/Private/ProfileTree/LaunchProfileTreeData.cpp
@@ -63,10 +63,15 @@ namespace ProjectLauncher
 			ExtensionInstance->CustomizeTree(*this);
 		}
 	}
+	
+	void FLaunchProfileTreeData::RequestTreeRefresh()
+	{
+		bRequestTreeRefresh = true;
+	}
 
 
 
-	FLaunchProfileTreeNode::FLaunchProfileTreeNode( const FLaunchProfileTreeData* InTreeData ) 
+	FLaunchProfileTreeNode::FLaunchProfileTreeNode( FLaunchProfileTreeData* InTreeData ) 
 		: TreeData(InTreeData)
 	{
 	}
@@ -79,6 +84,7 @@ namespace ProjectLauncher
 		TreeNode->Callbacks = InWidgetCallbacks;
 		Children.Add(TreeNode);
 
+		TreeData->RequestTreeRefresh();
 		return *this;
 	}
 
diff --git a/Engine/Plugins/Developer/ProjectLauncher/Source/ProjectLauncher/Private/Widgets/Profiles/SCustomLaunchCustomProfileEditor.cpp b/Engine/Plugins/Developer/ProjectLauncher/Source/ProjectLauncher/Private/Widgets/Profiles/SCustomLaunchCustomProfileEditor.cpp
--- a/Engine/Plugins/Developer/ProjectLauncher/Source/ProjectLauncher/Private/Widgets/Profiles/SCustomLaunchCustomProfileEditor.cpp
+++ b/Engine/Plugins/Developer/ProjectLauncher/Source/ProjectLauncher/Private/Widgets/Profiles/SCustomLaunchCustomProfileEditor.cpp
@@ -153,7 +153,7 @@ public:
 			auto SetToDefault = [TreeNode]()
 			{
 				TreeNode->Callbacks.SetToDefault();
-				TreeNode->TreeData->TreeBuilder->OnPropertyChanged();
+				TreeNode->GetTreeData()->TreeBuilder->OnPropertyChanged();
 				return FReply::Handled();
 			};
 
@@ -414,6 +414,18 @@ TSharedRef<SWidget> SCustomLaunchCustomProfileEditor::MakeExtensionsMenu()
 }
 
 
+void SCustomLaunchCustomProfileEditor::Tick(const FGeometry& AllottedGeometry, const double InCurrentTime, const float InDeltaTime)
+{
+	SCompoundWidget::Tick(AllottedGeometry, InCurrentTime, InDeltaTime);
+
+	if (TreeBuilder.IsValid() && TreeBuilder->GetProfileTree()->bRequestTreeRefresh)
+	{
+		TreeView->RequestTreeRefresh();
+		TreeBuilder->GetProfileTree()->bRequestTreeRefresh = false;
+	}
+}
+
+
 
 
 #undef LOCTEXT_NAMESPACE
diff --git a/Engine/Plugins/Developer/ProjectLauncher/Source/ProjectLauncher/Private/Widgets/Profiles/SCustomLaunchCustomProfileEditor.h b/Engine/Plugins/Developer/ProjectLauncher/Source/ProjectLauncher/Private/Widgets/Profiles/SCustomLaunchCustomProfileEditor.h
--- a/Engine/Plugins/Developer/ProjectLauncher/Source/ProjectLauncher/Private/Widgets/Profiles/SCustomLaunchCustomProfileEditor.h
+++ b/Engine/Plugins/Developer/ProjectLauncher/Source/ProjectLauncher/Private/Widgets/Profiles/SCustomLaunchCustomProfileEditor.h
@@ -41,6 +41,7 @@ public:
 	virtual TSharedRef<SWidget> GetScrollWidget() override;
 
 private:
+	virtual void Tick(const FGeometry& AllottedGeometry, const double InCurrentTime, const float InDeltaTime) override;
 
 	TSharedRef<ITableRow> OnGenerateWidgetForTreeNode( ProjectLauncher::FLaunchProfileTreeNodePtr TreeNode, const TSharedRef<STableViewBase>& OwnerTable );
 	void OnGetChildren(ProjectLauncher::FLaunchProfileTreeNodePtr TreeNode, TArray<ProjectLauncher::FLaunchProfileTreeNodePtr>& OutChildren);
diff --git a/Engine/Plugins/Developer/ProjectLauncher/Source/ProjectLauncher/Public/ProfileTree/LaunchProfileTreeData.h b/Engine/Plugins/Developer/ProjectLauncher/Source/ProjectLauncher/Public/ProfileTree/LaunchProfileTreeData.h
--- a/Engine/Plugins/Developer/ProjectLauncher/Source/ProjectLauncher/Public/ProfileTree/LaunchProfileTreeData.h
+++ b/Engine/Plugins/Developer/ProjectLauncher/Source/ProjectLauncher/Public/ProfileTree/LaunchProfileTreeData.h
@@ -29,7 +29,7 @@ namespace ProjectLauncher
 	{
 	public:
 		// construction
-		FLaunchProfileTreeNode( const FLaunchProfileTreeData* InTreeData );
+		FLaunchProfileTreeNode( FLaunchProfileTreeData* InTreeData );
 
 		typedef TFunction<void()> FFunction;
 		typedef TFunction<bool()> FGetBool;
@@ -82,10 +82,10 @@ namespace ProjectLauncher
 
 
 		TArray<FLaunchProfileTreeNodePtr> Children;
-		const FLaunchProfileTreeData* TreeData;
+		const FLaunchProfileTreeData* GetTreeData() const { return TreeData; }
 
 	protected:
-
+		FLaunchProfileTreeData* TreeData;
 	};
 
 
@@ -100,6 +100,8 @@ namespace ProjectLauncher
 		FLaunchProfileTreeNode& AddHeading( FText InName );
 		void CreateExtensionsUI();
 
+		void RequestTreeRefresh();
+
 		// member variables
 		ILauncherProfilePtr Profile;
 		TSharedRef<FModel> Model;
@@ -107,6 +109,7 @@ namespace ProjectLauncher
 		ILaunchProfileTreeBuilder* TreeBuilder;
 		TArray<TSharedPtr<FLaunchExtensionInstance>> ExtensionInstances;
 		bool bHasAnyMenuExtensions = false;
+		bool bRequestTreeRefresh = false;
 	};
 
 
