Only show build configurations that are available

#rnx

diff --git a/Engine/Plugins/Developer/ProjectLauncher/Source/ProjectLauncher/Private/ProfileTree/GenericProfileTreeBuilder.cpp b/Engine/Plugins/Developer/ProjectLauncher/Source/ProjectLauncher/Private/ProfileTree/GenericProfileTreeBuilder.cpp
--- a/Engine/Plugins/Developer/ProjectLauncher/Source/ProjectLauncher/Private/ProfileTree/GenericProfileTreeBuilder.cpp
+++ b/Engine/Plugins/Developer/ProjectLauncher/Source/ProjectLauncher/Private/ProfileTree/GenericProfileTreeBuilder.cpp
@@ -19,6 +19,7 @@
 #include "Interfaces/ITargetPlatformManagerModule.h"
 #include "Interfaces/ITargetPlatform.h"
 #include "GameProjectHelper.h"
+#include "InstalledPlatformInfo.h"
 
 
 
@@ -183,11 +184,23 @@ namespace ProjectLauncher
 
 	void FGenericProfileTreeBuilder::AddConfigurationProperty( FLaunchProfileTreeNode& HeadingNode )
 	{
+		TArray<EBuildConfiguration> ValidConfigurations;
+		
+		static const EBuildConfiguration AllConfigurations[] = { EBuildConfiguration::Debug, EBuildConfiguration::DebugGame, EBuildConfiguration::Development, EBuildConfiguration::Test, EBuildConfiguration::Shipping };
+		for (EBuildConfiguration Configuration : AllConfigurations)
+		{
+			// only show the configurations that are currently available. @todo: might be better to show all, but disable the ones that are unavailale
+			if (FInstalledPlatformInfo::Get().IsValidConfiguration(Configuration))
+			{
+				ValidConfigurations.Add(Configuration);
+			}
+		}
+
 		HeadingNode.AddWidget( LOCTEXT("ConfigurationLabel", "Configuration"), 
 			SNew(SCustomLaunchLexToStringCombo<EBuildConfiguration>)
 			.OnSelectionChanged(this, &FGenericProfileTreeBuilder::SetBuildConfiguration)
 			.SelectedItem(this, &FGenericProfileTreeBuilder::GetBuildConfiguration)
-			.Items( TArray<EBuildConfiguration>({ EBuildConfiguration::Debug, EBuildConfiguration::Development, EBuildConfiguration::Test, EBuildConfiguration::Shipping }) )
+			.Items(ValidConfigurations)
 		);
 	}
 
