FModel::GetBuildTargetInfo uses the default build target if the profile's build target has not been overridden

#rnx

diff --git a/Engine/Plugins/Developer/ProjectLauncher/Source/ProjectLauncher/Private/Model/ProjectLauncherModel.cpp b/Engine/Plugins/Developer/ProjectLauncher/Source/ProjectLauncher/Private/Model/ProjectLauncherModel.cpp
--- a/Engine/Plugins/Developer/ProjectLauncher/Source/ProjectLauncher/Private/Model/ProjectLauncherModel.cpp
+++ b/Engine/Plugins/Developer/ProjectLauncher/Source/ProjectLauncher/Private/Model/ProjectLauncherModel.cpp
@@ -247,8 +247,13 @@ namespace ProjectLauncher
 
 
 
-	FTargetInfo FModel::GetBuildTargetInfo( const FString& BuildTargetName, const FString& ProjectPath )
+	FTargetInfo FModel::GetBuildTargetInfo( FString BuildTargetName, const FString& ProjectPath )
 	{
+		if (BuildTargetName.IsEmpty())
+		{
+			BuildTargetName = GetProjectSettings(ProjectPath).DefaultBuildTargetName;
+		}
+
 		FTargetInfo Result;
 		if (!BuildTargetName.IsEmpty())
 		{
@@ -575,12 +580,28 @@ namespace ProjectLauncher
 	}
 
 
-	void FModel::ReadProjectSettingsFromConfig( FConfigCacheIni& InConfig, FProjectSettings& OutResult )
+	void FModel::ReadProjectSettingsFromConfig( FConfigCacheIni& InConfig, const FString& InProjectPath, FProjectSettings& OutResult )
 	{
 		// read project packaging settings
 		const TCHAR* ProjectPackagingConfigSection = TEXT("/Script/UnrealEd.ProjectPackagingSettings");
 		InConfig.GetBool(ProjectPackagingConfigSection, TEXT("bUseZenStore"), OutResult.bUseZenStore, GGameIni);
 		InConfig.GetBool(ProjectPackagingConfigSection, TEXT("bEnablePakStreaming"), OutResult.bHasAutomaticZenPakStreamingWorkspaceCreation, GGameIni);
+		InConfig.GetString(ProjectPackagingConfigSection, TEXT("BuildTarget"), OutResult.DefaultBuildTargetName, GGameIni);
+		if (OutResult.DefaultBuildTargetName.IsEmpty())
+		{
+			const TArray<FTargetInfo>& BuildTargets = FDesktopPlatformModule::Get()->GetTargetsForProject(*InProjectPath);
+			const TArray<FTargetInfo> GameBuildTargets = BuildTargets.FilterByPredicate([](const FTargetInfo& TargetInfo)
+			{
+				return TargetInfo.Type == EBuildTargetType::Game;
+			});
+			if (GameBuildTargets.Num() == 1)
+			{
+				OutResult.DefaultBuildTargetName = GameBuildTargets[0].Name;
+			}
+			// multiple (or no) game targets... would need to read the BuildTarget from /Script/BuildSettings.BuildSettings... also reading the platform-specific ini
+			// ... fixme ...
+		}
+
 		
 		// read zen settings
 #if UE_WITH_ZEN
@@ -610,11 +631,15 @@ namespace ProjectLauncher
 		{
 			// read the current project's ini settings
 			Result.bIsCurrentEditorProject = true;
-			ReadProjectSettingsFromConfig(*GConfig, Result);
+			ReadProjectSettingsFromConfig(*GConfig, InProjectPath, Result);
 
 			// apply any in-memory properties @fixme: is this necessary?
 			const UProjectPackagingSettings* ProjectPackagingSettings = UProjectPackagingSettings::StaticClass()->GetDefaultObject<UProjectPackagingSettings>();
 			Result.bUseZenStore = ProjectPackagingSettings->bUseZenStore;
+			if (!ProjectPackagingSettings->BuildTarget.IsEmpty())
+			{
+				Result.DefaultBuildTargetName = ProjectPackagingSettings->BuildTarget;
+			}
 
 			// note this is not cached because the user could edit the properties from within the editor
 		}
@@ -632,7 +657,7 @@ namespace ProjectLauncher
 
 			// read the ini settings
 			Result.bIsCurrentEditorProject = false;
-			ReadProjectSettingsFromConfig(Config, Result);
+			ReadProjectSettingsFromConfig(Config, InProjectPath, Result);
 
 			// cache the result
 			CachedProjectSettings.Add( ProjectName, Result);
diff --git a/Engine/Plugins/Developer/ProjectLauncher/Source/ProjectLauncher/Private/ProfileTree/GenericProfileTreeBuilder.cpp b/Engine/Plugins/Developer/ProjectLauncher/Source/ProjectLauncher/Private/ProfileTree/GenericProfileTreeBuilder.cpp
--- a/Engine/Plugins/Developer/ProjectLauncher/Source/ProjectLauncher/Private/ProfileTree/GenericProfileTreeBuilder.cpp
+++ b/Engine/Plugins/Developer/ProjectLauncher/Source/ProjectLauncher/Private/ProfileTree/GenericProfileTreeBuilder.cpp
@@ -146,7 +146,7 @@ namespace ProjectLauncher
 					.IsDefault = [this]()		{ return !Profile->HasBuildTargetSpecified() || Profile->GetBuildTarget().IsEmpty(); },
 					.SetToDefault = [this]()	{ SetBuildTarget(FString()); },
 				},
-				SNew(SCustomLaunchBuildTargetCombo)
+				SNew(SCustomLaunchBuildTargetCombo, Model)
 				.OnSelectionChanged(this, &FGenericProfileTreeBuilder::SetBuildTarget)
 				.SelectedBuildTarget(this, &FGenericProfileTreeBuilder::GetBuildTarget)
 				.SelectedProject(this, &FGenericProfileTreeBuilder::GetProjectPath)
@@ -161,7 +161,7 @@ namespace ProjectLauncher
 					.SetToDefault = [this]()	{ SetBuildTarget(FString()); },
 					.IsEnabled = [this]()		{ return Profile->HasProjectSpecified(); },
 				},
-				SNew(SCustomLaunchBuildTargetCombo)
+				SNew(SCustomLaunchBuildTargetCombo, Model)
 				.OnSelectionChanged(this, &FGenericProfileTreeBuilder::SetBuildTarget)
 				.SelectedBuildTarget(this, &FGenericProfileTreeBuilder::GetBuildTarget)
 				.SupportedTargetTypes(this, &FGenericProfileTreeBuilder::GetSupportedBuildTargetTypes)
@@ -631,7 +631,7 @@ namespace ProjectLauncher
 
 		if (bUseFriendlyBuildTargetSelection)
 		{
-			FTargetInfo BuildTargetInfo = FModel::GetBuildTargetInfo(GetBuildTarget(), GetProjectPath());
+			FTargetInfo BuildTargetInfo = Model->GetBuildTargetInfo(GetBuildTarget(), GetProjectPath());
 
 			for (const FString& Platform : SelectedPlatforms)
 			{
diff --git a/Engine/Plugins/Developer/ProjectLauncher/Source/ProjectLauncher/Private/Widgets/Profiles/SCustomLaunchProfilesPanel.cpp b/Engine/Plugins/Developer/ProjectLauncher/Source/ProjectLauncher/Private/Widgets/Profiles/SCustomLaunchProfilesPanel.cpp
--- a/Engine/Plugins/Developer/ProjectLauncher/Source/ProjectLauncher/Private/Widgets/Profiles/SCustomLaunchProfilesPanel.cpp
+++ b/Engine/Plugins/Developer/ProjectLauncher/Source/ProjectLauncher/Private/Widgets/Profiles/SCustomLaunchProfilesPanel.cpp
@@ -176,7 +176,7 @@ TSharedRef<SWidget> SCustomLaunchProfilesPanel::CreateProfilesPanel()
 			.VAlign(VAlign_Center)
 			.Padding(2,0)
 			[
-				SNew(SCustomLaunchBuildTargetCombo)
+				SNew(SCustomLaunchBuildTargetCombo, Model.ToSharedRef())
 				.OnSelectionChanged(this, &SCustomLaunchProfilesPanel::SetDefaultBuildTarget)
 				.SelectedBuildTarget(this, &SCustomLaunchProfilesPanel::GetDefaultBuildTarget)
 				.SelectedProject(this, &SCustomLaunchProfilesPanel::GetDefaultProjectPath)
diff --git a/Engine/Plugins/Developer/ProjectLauncher/Source/ProjectLauncher/Private/Widgets/Shared/SCustomLaunchBuildTargetCombo.cpp b/Engine/Plugins/Developer/ProjectLauncher/Source/ProjectLauncher/Private/Widgets/Shared/SCustomLaunchBuildTargetCombo.cpp
--- a/Engine/Plugins/Developer/ProjectLauncher/Source/ProjectLauncher/Private/Widgets/Shared/SCustomLaunchBuildTargetCombo.cpp
+++ b/Engine/Plugins/Developer/ProjectLauncher/Source/ProjectLauncher/Private/Widgets/Shared/SCustomLaunchBuildTargetCombo.cpp
@@ -13,12 +13,13 @@
 
 
 BEGIN_SLATE_FUNCTION_BUILD_OPTIMIZATION
-void SCustomLaunchBuildTargetCombo::Construct(const FArguments& InArgs)
+void SCustomLaunchBuildTargetCombo::Construct(const FArguments& InArgs, TSharedRef<ProjectLauncher::FModel> InModel)
 {
 	OnSelectionChanged = InArgs._OnSelectionChanged;
 	SelectedProject = InArgs._SelectedProject;
 	SelectedBuildTarget = InArgs._SelectedBuildTarget;
 	SupportedTargetTypes = InArgs._SupportedTargetTypes;
+	Model = InModel;
 
 	FSlateFontInfo Font = InArgs._Font.IsSet() ? InArgs._Font.Get() : InArgs._TextStyle->Font;
 
@@ -134,29 +135,7 @@ FString SCustomLaunchBuildTargetCombo::GetDefaultBuildTargetName() const
 {
 	if (SelectedProject.IsSet())
 	{
-		// if this is the current in-editor project, we can check the default setting
-		if (FPaths::IsProjectFilePathSet() && SelectedProject.Get() == FPaths::GetProjectFilePath())
-		{
-			UProjectPackagingSettings* AllPlatformPackagingSettings = GetMutableDefault<UProjectPackagingSettings>();
-			if (AllPlatformPackagingSettings != nullptr && !AllPlatformPackagingSettings->BuildTarget.IsEmpty())
-			{
-				return AllPlatformPackagingSettings->BuildTarget;
-			}
-		}
-
-		// if there is only one game build target, it will be used by default
-		const TArray<FTargetInfo>& BuildTargets = FDesktopPlatformModule::Get()->GetTargetsForProject(SelectedProject.Get());
-		const TArray<FTargetInfo> GameBuildTargets = BuildTargets.FilterByPredicate([](const FTargetInfo& TargetInfo)
-		{
-			return TargetInfo.Type == EBuildTargetType::Game;
-		});
-		if (GameBuildTargets.Num() == 1)
-		{
-			return GameBuildTargets[0].Name;
-		}
-
-		// multiple (or no) game targets... would need to read the BuildTarget from the other project's ini file
-		// ... fixme ...
+		return Model->GetProjectSettings(SelectedProject.Get()).DefaultBuildTargetName;
 	}
 
 	return FString();
diff --git a/Engine/Plugins/Developer/ProjectLauncher/Source/ProjectLauncher/Public/Model/ProjectLauncherModel.h b/Engine/Plugins/Developer/ProjectLauncher/Source/ProjectLauncher/Public/Model/ProjectLauncherModel.h
--- a/Engine/Plugins/Developer/ProjectLauncher/Source/ProjectLauncher/Public/Model/ProjectLauncherModel.h
+++ b/Engine/Plugins/Developer/ProjectLauncher/Source/ProjectLauncher/Public/Model/ProjectLauncherModel.h
@@ -72,6 +72,9 @@ namespace ProjectLauncher
 
 		// whether Zen Server is allowed to accept connections from remote machines (i.e. console devkits, phones etc)
 		bool bAllowRemoteNetworkService = false;
+
+		// default build target name
+		FString DefaultBuildTargetName;
 	};
 
 
@@ -155,12 +158,12 @@ namespace ProjectLauncher
 		static bool IsHostPlatform(const ILauncherProfilePtr& Profile);
 		static bool IsHostPlatform(FName PlatformName);
 
-		static FTargetInfo GetBuildTargetInfo( const FString& BuildTargetName, const FString& ProjectPath );
-		static FTargetInfo GetBuildTargetInfo( const ILauncherProfileRef& Profile );
+		FTargetInfo GetBuildTargetInfo( FString BuildTargetName, const FString& ProjectPath );
+		FTargetInfo GetBuildTargetInfo( const ILauncherProfileRef& Profile );
 
 		static TSharedPtr<ITargetDeviceProxy> GetDeviceProxy(const ILauncherProfileRef& Profile);
-		static void UpdatedCookedPlatformsFromDeployDeviceProxy(const ILauncherProfileRef& Profile, TSharedPtr<ITargetDeviceProxy> DeviceProxy = nullptr);
-		static void UpdateCookedPlatformsFromBuildTarget(const ILauncherProfileRef& Profile);
+		void UpdatedCookedPlatformsFromDeployDeviceProxy(const ILauncherProfileRef& Profile, TSharedPtr<ITargetDeviceProxy> DeviceProxy = nullptr);
+		void UpdateCookedPlatformsFromBuildTarget(const ILauncherProfileRef& Profile);
 
 		static FString GetVanillaPlatformName( const FString& PlatformName );
 		static FString GetBuildTargetPlatformName( const FString& PlatformName, const FTargetInfo& BuildTargetInfo );
@@ -189,7 +192,7 @@ namespace ProjectLauncher
 		void HandleDeviceProxyAdded(const TSharedRef<ITargetDeviceProxy>& DeviceProxy );
 		void HandleDeviceProxyRemoved(const TSharedRef<ITargetDeviceProxy>& DeviceProxy );
 
-		void ReadProjectSettingsFromConfig( FConfigCacheIni& InConfig, FProjectSettings& OutResult );
+		void ReadProjectSettingsFromConfig( FConfigCacheIni& InConfig, const FString& InProjectPath, FProjectSettings& OutResult );
 
 		void OnModifyLaunchCommandLine( const ILauncherProfileRef& InProfile, FString& InOutCommandLine );
 
diff --git a/Engine/Plugins/Developer/ProjectLauncher/Source/ProjectLauncher/Public/Widgets/Shared/SCustomLaunchBuildTargetCombo.h b/Engine/Plugins/Developer/ProjectLauncher/Source/ProjectLauncher/Public/Widgets/Shared/SCustomLaunchBuildTargetCombo.h
--- a/Engine/Plugins/Developer/ProjectLauncher/Source/ProjectLauncher/Public/Widgets/Shared/SCustomLaunchBuildTargetCombo.h
+++ b/Engine/Plugins/Developer/ProjectLauncher/Source/ProjectLauncher/Public/Widgets/Shared/SCustomLaunchBuildTargetCombo.h
@@ -8,6 +8,7 @@
 #include "Styling/CoreStyle.h"
 #include "Styling/SlateWidgetStyleAsset.h"
 #include "Fonts/SlateFontInfo.h"
+#include "Model/ProjectLauncherModel.h"
 
 class PROJECTLAUNCHER_API SCustomLaunchBuildTargetCombo
 	: public SCompoundWidget
@@ -27,7 +28,7 @@ public:
 	SLATE_END_ARGS()
 
 public:
-	void Construct(	const FArguments& InArgs);
+	void Construct(	const FArguments& InArgs, TSharedRef<ProjectLauncher::FModel> InModel);
 
 protected:
 	FString GetDefaultBuildTargetName() const;
@@ -37,6 +38,7 @@ protected:
 	TAttribute<FString> SelectedBuildTarget;
 	FOnSelectionChanged OnSelectionChanged;
 	bool bShowAnyProjectOption;
+	TSharedPtr<ProjectLauncher::FModel> Model;
 
 	TSharedRef<SWidget> MakeBuildTargetSelectionWidget();
 
