ProjectLauncher2 Backend extension - retrieve the correct platform name

diff --git a/Engine/Plugins/Developer/ProjectLauncher/Source/ProjectLauncher/Private/Extension/LaunchExtension.cpp b/Engine/Plugins/Developer/ProjectLauncher/Source/ProjectLauncher/Private/Extension/LaunchExtension.cpp
--- a/Engine/Plugins/Developer/ProjectLauncher/Source/ProjectLauncher/Private/Extension/LaunchExtension.cpp
+++ b/Engine/Plugins/Developer/ProjectLauncher/Source/ProjectLauncher/Private/Extension/LaunchExtension.cpp
@@ -11,6 +11,7 @@
 #include "Misc/ConfigCacheIni.h"
 #include "Modules/ModuleManager.h"
 #include "ILauncherServicesModule.h"
+#include "PlatformInfo.h"
 
 #define LOCTEXT_NAMESPACE "FLaunchExtensionInstance"
 
@@ -382,8 +383,10 @@ namespace ProjectLauncher
 		TreeBuilder->OnPropertyChanged();
 	}
 
-
-
+	FName FLaunchExtensionInstance::GetLaunchPlatformName() const
+	{
+		return GetModel()->GetPlatformInfo(GetProfile())->IniPlatformName;
+	}
 
 	FString FLaunchExtensionInstance::GetConfigString( EConfig Config, const TCHAR* Name, const TCHAR* DefaultValue ) const
 	{
diff --git a/Engine/Plugins/Developer/ProjectLauncher/Source/ProjectLauncher/Public/Extension/LaunchExtension.h b/Engine/Plugins/Developer/ProjectLauncher/Source/ProjectLauncher/Public/Extension/LaunchExtension.h
--- a/Engine/Plugins/Developer/ProjectLauncher/Source/ProjectLauncher/Public/Extension/LaunchExtension.h
+++ b/Engine/Plugins/Developer/ProjectLauncher/Source/ProjectLauncher/Public/Extension/LaunchExtension.h
@@ -172,6 +172,10 @@ namespace ProjectLauncher
 		 */
 		inline TSharedRef<FModel> GetModel() const { return Model; }
 
+		/**
+		 * Get the name of the platform to launch on
+		 */
+		FName GetLaunchPlatformName() const;
 
 		/** 
 		* Enumeration to choose where a value should be stored
