ProjectLauncher2 Backend extension - cache project config in the model so it can be passed into extension, this fixes the project set in the launch profile doesn't match the active project in editor which results to incorrect config being used to find backend issue.

diff --git a/Engine/Plugins/Developer/ProjectLauncher/Source/ProjectLauncher/Private/Model/ProjectLauncherModel.cpp b/Engine/Plugins/Developer/ProjectLauncher/Source/ProjectLauncher/Private/Model/ProjectLauncherModel.cpp
--- a/Engine/Plugins/Developer/ProjectLauncher/Source/ProjectLauncher/Private/Model/ProjectLauncherModel.cpp
+++ b/Engine/Plugins/Developer/ProjectLauncher/Source/ProjectLauncher/Private/Model/ProjectLauncherModel.cpp
@@ -78,6 +78,11 @@ namespace ProjectLauncher
 
 	FModel::~FModel()
 	{
+		for (TPair<FString, FProjectSettings>& Pair : CachedProjectSettings)
+		{
+			delete Pair.Value.Config;
+		}
+
 		ProfileManager->OnPostProcessLaunchCommandLine().RemoveAll( this );
 
 		ProfileManager->OnProfileAdded().RemoveAll(this);
@@ -631,6 +636,7 @@ namespace ProjectLauncher
 		{
 			// read the current project's ini settings
 			Result.bIsCurrentEditorProject = true;
+			Result.Config = GConfig;
 			ReadProjectSettingsFromConfig(*GConfig, InProjectPath, Result);
 
 			// apply any in-memory properties @fixme: is this necessary?
@@ -650,14 +656,14 @@ namespace ProjectLauncher
 		else if (!ProjectPath.IsEmpty())
 		{
 			// load the other project's ini files into a temporary config cache
-			FConfigCacheIni Config(EConfigCacheType::Temporary);
-			FConfigContext Context = FConfigContext::ReadIntoConfigSystem(&Config, FString());
+			Result.Config = new FConfigCacheIni(EConfigCacheType::Temporary);
+			FConfigContext Context = FConfigContext::ReadIntoConfigSystem(Result.Config, FString());
 			Context.ProjectConfigDir = FPaths::Combine(FPaths::GetPath(ProjectPath), TEXT("Config/"));
-			Config.InitializeKnownConfigFiles(Context);
+			Result.Config->InitializeKnownConfigFiles(Context);
 
 			// read the ini settings
 			Result.bIsCurrentEditorProject = false;
-			ReadProjectSettingsFromConfig(Config, InProjectPath, Result);
+			ReadProjectSettingsFromConfig(*Result.Config, InProjectPath, Result);
 
 			// cache the result
 			CachedProjectSettings.Add( ProjectName, Result);
@@ -671,8 +677,6 @@ namespace ProjectLauncher
 		return GetProjectSettings(*Profile->GetProjectPath());
 	}
 
-
-
 	bool FModel::IsBasicLaunchProfile( const ILauncherProfilePtr& Profile ) const
 	{
 		return (Profile == BasicLaunchProfile);
diff --git a/Engine/Plugins/Developer/ProjectLauncher/Source/ProjectLauncher/Public/Model/ProjectLauncherModel.h b/Engine/Plugins/Developer/ProjectLauncher/Source/ProjectLauncher/Public/Model/ProjectLauncherModel.h
--- a/Engine/Plugins/Developer/ProjectLauncher/Source/ProjectLauncher/Public/Model/ProjectLauncherModel.h
+++ b/Engine/Plugins/Developer/ProjectLauncher/Source/ProjectLauncher/Public/Model/ProjectLauncherModel.h
@@ -75,6 +75,8 @@ namespace ProjectLauncher
 
 		// default build target name
 		FString DefaultBuildTargetName;
+
+		FConfigCacheIni* Config = nullptr;
 	};
 
 
