Project Launcher extensions are created once and reused

#rnx

diff --git a/Engine/Plugins/Developer/ProjectLauncher/Source/ProjectLauncher/Private/Extension/LaunchExtension.cpp b/Engine/Plugins/Developer/ProjectLauncher/Source/ProjectLauncher/Private/Extension/LaunchExtension.cpp
--- a/Engine/Plugins/Developer/ProjectLauncher/Source/ProjectLauncher/Private/Extension/LaunchExtension.cpp
+++ b/Engine/Plugins/Developer/ProjectLauncher/Source/ProjectLauncher/Private/Extension/LaunchExtension.cpp
@@ -2,7 +2,6 @@
 
 #include "Extension/LaunchExtension.h"
 #include "Model/ProjectLauncherModel.h"
-#include "ProfileTree/ILaunchProfileTreeBuilder.h"
 #include "Containers/Array.h"
 #include "Templates/SharedPointer.h"
 #include "HAL/CriticalSection.h"
@@ -19,6 +18,7 @@
 namespace ProjectLauncher
 {
 	static TArray<TSharedPtr<FLaunchExtension>> GExtensions;
+	static TMap<ILauncherProfilePtr, TArray<TSharedPtr<FLaunchExtensionInstance>>> GExtensionInstances;
 	static FCriticalSection GExtensionsCS;
 
 
@@ -37,25 +37,19 @@ namespace ProjectLauncher
 
 	void ApplyExtensionVariables( const ILauncherProfileRef& InProfile, FString& InOutCommandLine, TSharedRef<FModel> InModel )
 	{
-		// take a snapshot of the extension list
-		TArray<TSharedPtr<FLaunchExtension>> Extensions;
+		// take a snapshot of the extension instances for this profile
+		TArray<TSharedPtr<FLaunchExtensionInstance>> ExtensionInstances;
 		{
 			FScopeLock Lock(&GExtensionsCS);
-			Extensions = GExtensions;
+			if (GExtensionInstances.Contains(InProfile))
+			{
+				ExtensionInstances = GExtensionInstances[InProfile];
+			}
 		}
 
 
-		for (TSharedPtr<FLaunchExtension> Extension : Extensions)
+		for (TSharedPtr<FLaunchExtensionInstance> ExtensionInstance : ExtensionInstances)
 		{
-			// instantiate the extension and apply the extension variables
-			FLaunchExtensionInstance::FArgs Args
-			{
-				.Profile = InProfile,
-				.TreeBuilder = nullptr,
-				.Model = InModel,
-				.Extension = Extension.ToSharedRef(),
-			};
-			TSharedPtr<FLaunchExtensionInstance> ExtensionInstance = Extension->CreateInstanceForProfile(Args);
 			if (ExtensionInstance.IsValid())
 			{
 				// apply variable substitutions
@@ -82,10 +76,19 @@ namespace ProjectLauncher
 	}
 
 
+	void RemoveExtensionInstancesForProfile( const ILauncherProfileRef& Profile )
+	{
+		FScopeLock Lock(&GExtensionsCS);
+		GExtensionInstances.Remove(Profile.ToSharedPtr());
+	}
+
 
 
-	TArray<TSharedPtr<FLaunchExtensionInstance>> FLaunchExtension::CreateExtensionInstancesForProfile( ILauncherProfileRef InProfile, TSharedRef<FModel> InModel, ILaunchProfileTreeBuilder* InTreeBuilder )
+	TArray<TSharedPtr<FLaunchExtensionInstance>> FLaunchExtension::CreateExtensionInstancesForProfile( ILauncherProfileRef InProfile, TSharedRef<FModel> InModel )
 	{
+		// make sure the extensions for this profile haven't been created already
+		check(!GExtensionInstances.Contains(InProfile.ToSharedPtr()));
+
 		// take a snapshot of the extension list
 		TArray<TSharedPtr<FLaunchExtension>> Extensions;
 		{
@@ -100,7 +103,6 @@ namespace ProjectLauncher
 			FLaunchExtensionInstance::FArgs Args
 			{
 				.Profile = InProfile,
-				.TreeBuilder = InTreeBuilder,
 				.Model = InModel,
 				.Extension = Extension.ToSharedRef(),
 			};
@@ -112,6 +114,11 @@ namespace ProjectLauncher
 			}
 		}
 
+		// record the extension instances for the this profile
+		{
+			FScopeLock Lock(&GExtensionsCS);
+			GExtensionInstances.Add(InProfile.ToSharedPtr(), Result);
+		}
 		return Result;
 	}
 
@@ -127,7 +134,6 @@ namespace ProjectLauncher
 
 	FLaunchExtensionInstance::FLaunchExtensionInstance( FArgs& InArgs )
 		: Profile(InArgs.Profile)
-		, TreeBuilder(InArgs.TreeBuilder)
 		, Model(InArgs.Model)
 		, Extension(InArgs.Extension)
 	{
@@ -380,7 +386,7 @@ namespace ProjectLauncher
 	void FLaunchExtensionInstance::SetCommandLine( const FString& CommandLine )
 	{
 		Profile->SetAdditionalCommandLineParameters(CommandLine);
-		TreeBuilder->OnPropertyChanged();
+		BroadcastPropertyChanged();
 	}
 
 	FName FLaunchExtensionInstance::GetLaunchPlatformName() const
@@ -436,7 +442,7 @@ namespace ProjectLauncher
 			{
 				Profile->GetCustomStringProperties().Add(KeyName, Value);
 			}
-			TreeBuilder->OnPropertyChanged();
+			BroadcastPropertyChanged();
 		}
 		else
 		{
@@ -457,7 +463,7 @@ namespace ProjectLauncher
 			{
 				Profile->GetCustomBoolProperties().Add(KeyName, Value);
 			}
-			TreeBuilder->OnPropertyChanged();
+			BroadcastPropertyChanged();
 		}
 		else
 		{
@@ -487,6 +493,11 @@ namespace ProjectLauncher
 		checkNoEntry();
 		return Name;
 	}
+
+	void FLaunchExtensionInstance::BroadcastPropertyChanged() const
+	{
+		PropertyChangedDelegate.Broadcast();
+	}
 }
 
 #undef LOCTEXT_NAMESPACE
diff --git a/Engine/Plugins/Developer/ProjectLauncher/Source/ProjectLauncher/Private/Model/ProjectLauncherModel.cpp b/Engine/Plugins/Developer/ProjectLauncher/Source/ProjectLauncher/Private/Model/ProjectLauncherModel.cpp
--- a/Engine/Plugins/Developer/ProjectLauncher/Source/ProjectLauncher/Private/Model/ProjectLauncherModel.cpp
+++ b/Engine/Plugins/Developer/ProjectLauncher/Source/ProjectLauncher/Private/Model/ProjectLauncherModel.cpp
@@ -332,10 +332,13 @@ namespace ProjectLauncher
 		AllProfiles.Add(Profile);
 	}
 
+	extern void RemoveExtensionInstancesForProfile( const ILauncherProfileRef& Profile );
 
 
 	void FModel::HandleProfileManagerProfileRemoved(const ILauncherProfileRef& Profile)
 	{
+		RemoveExtensionInstancesForProfile(Profile);
+
 		AllProfiles.Remove(Profile);
 
 		if (Profile == SelectedProfile)
diff --git a/Engine/Plugins/Developer/ProjectLauncher/Source/ProjectLauncher/Private/ProfileTree/LaunchProfileTreeData.cpp b/Engine/Plugins/Developer/ProjectLauncher/Source/ProjectLauncher/Private/ProfileTree/LaunchProfileTreeData.cpp
--- a/Engine/Plugins/Developer/ProjectLauncher/Source/ProjectLauncher/Private/ProfileTree/LaunchProfileTreeData.cpp
+++ b/Engine/Plugins/Developer/ProjectLauncher/Source/ProjectLauncher/Private/ProfileTree/LaunchProfileTreeData.cpp
@@ -34,11 +34,13 @@ namespace ProjectLauncher
 
 		if (InProfile.IsValid())
 		{
-			ExtensionInstances = FLaunchExtension::CreateExtensionInstancesForProfile(Profile.ToSharedRef(), InModel, InTreeBuilder);
+			ExtensionInstances = FLaunchExtension::CreateExtensionInstancesForProfile(Profile.ToSharedRef(), InModel);
 
 			bHasAnyMenuExtensions = false;
 			for (TSharedPtr<FLaunchExtensionInstance> ExtensionInstance : ExtensionInstances)
 			{
+				ExtensionInstance->GetPropertyChangedDelegate().AddRaw(TreeBuilder, &ILaunchProfileTreeBuilder::OnPropertyChanged );
+
 				if (ExtensionInstance->HasCustomExtensionMenu())
 				{
 					bHasAnyMenuExtensions = true;
@@ -47,6 +49,15 @@ namespace ProjectLauncher
 		}
 	}
 
+	FLaunchProfileTreeData::~FLaunchProfileTreeData()
+	{
+		for (TSharedPtr<FLaunchExtensionInstance> ExtensionInstance : ExtensionInstances)
+		{
+			ExtensionInstance->GetPropertyChangedDelegate().RemoveAll(TreeBuilder);
+		}
+	}
+
+
 	FLaunchProfileTreeNode& FLaunchProfileTreeData::AddHeading( FText InName )
 	{
 		FLaunchProfileTreeNodePtr TreeNode = MakeShared<FLaunchProfileTreeNode>(this);
diff --git a/Engine/Plugins/Developer/ProjectLauncher/Source/ProjectLauncher/Public/Extension/LaunchExtension.h b/Engine/Plugins/Developer/ProjectLauncher/Source/ProjectLauncher/Public/Extension/LaunchExtension.h
--- a/Engine/Plugins/Developer/ProjectLauncher/Source/ProjectLauncher/Public/Extension/LaunchExtension.h
+++ b/Engine/Plugins/Developer/ProjectLauncher/Source/ProjectLauncher/Public/Extension/LaunchExtension.h
@@ -6,6 +6,7 @@
 #include "ProfileTree/LaunchProfileTreeData.h"
 #include "ILauncherProfile.h"
 #include "Framework/MultiBox/MultiBoxBuilder.h"
+#include "Delegates/Delegate.h"
 
 #define UE_API PROJECTLAUNCHER_API
 
@@ -13,6 +14,8 @@ namespace ProjectLauncher
 {
 	class FLaunchExtension;
 
+	typedef TMulticastDelegate<void()> FPropertyChangedDelegate;
+
 	/**
 	 * Base class for a launch extension instance.
 	 * Used while editing a specific profile, and when finalizing the command line arguments during profile launch.
@@ -30,7 +33,6 @@ namespace ProjectLauncher
 		struct FArgs
 		{
 			ILauncherProfileRef Profile;
-			ILaunchProfileTreeBuilder* TreeBuilder;
 			TSharedRef<FModel> Model;
 			TSharedRef<FLaunchExtension> Extension;
 		};
@@ -91,6 +93,11 @@ namespace ProjectLauncher
 		 */
 		inline TSharedRef<FLaunchExtension> GetExtension() const { return Extension; }
 
+		/**
+		 * Get the property changed delegate that will be called when an extension changes the profile
+		 */
+		inline FPropertyChangedDelegate& GetPropertyChangedDelegate() { return PropertyChangedDelegate; }
+
 	protected:
 
 		/**
@@ -237,7 +244,10 @@ namespace ProjectLauncher
 		 */
 		UE_API FString GetConfigKeyName( EConfig Config, const TCHAR* Name ) const;
 
-
+		/**
+		 * Signal to any listeners that a property in the profile has changed
+		 */
+		void BroadcastPropertyChanged() const;
 
 	private:
 		UE_API bool IsParameterGroup( const FString& InParameter ) const;
@@ -245,9 +255,9 @@ namespace ProjectLauncher
 
 
 		ILauncherProfileRef Profile;
-		ILaunchProfileTreeBuilder* TreeBuilder;
 		TSharedRef<FModel> Model;
 		TSharedRef<FLaunchExtension> Extension;
+		FPropertyChangedDelegate PropertyChangedDelegate;
 
 		friend class FLaunchProfileTreeNode;
 		UE_API void MakeCommandLineSubmenu( FMenuBuilder& MenuBuilder );
@@ -269,13 +279,6 @@ namespace ProjectLauncher
 	public:
 		virtual ~FLaunchExtension() = default;
 
-		/**
-		 * Create an instance of the launch extension for the given profile
-		 *
-		 * @returns: new instance, or null if it isn't appropriate
-		 */
-		virtual TSharedPtr<FLaunchExtensionInstance> CreateInstanceForProfile( FLaunchExtensionInstance::FArgs& InArgs ) = 0;
-
 		/**
 		 * Returns the debug name for this extension
 		 */
@@ -291,11 +294,19 @@ namespace ProjectLauncher
 		 * Instantiate all compatible extensions
 		 * 
 		 * @param InProfile the current profile 
-		 * @param InTreeBuilder the active tree builder
 		 * @param InModel helper class
 		 * @returns array of all compatible extensions
 		 */
-		static UE_API TArray<TSharedPtr<FLaunchExtensionInstance>> CreateExtensionInstancesForProfile( ILauncherProfileRef InProfile, TSharedRef<FModel> InModel, ILaunchProfileTreeBuilder* InTreeBuilder );
+		UE_API static TArray<TSharedPtr<FLaunchExtensionInstance>> CreateExtensionInstancesForProfile( ILauncherProfileRef InProfile, TSharedRef<FModel> InModel );
+
+
+	private:
+		/**
+		 * Create an instance of the launch extension for the given profile - called internally by CreateExtensionInstancesForProfile
+		 *
+		 * @returns: new instance, or null if it isn't appropriate
+		 */
+		virtual TSharedPtr<FLaunchExtensionInstance> CreateInstanceForProfile( FLaunchExtensionInstance::FArgs& InArgs ) = 0;
 	};
 
 
diff --git a/Engine/Plugins/Developer/ProjectLauncher/Source/ProjectLauncher/Public/ProfileTree/LaunchProfileTreeData.h b/Engine/Plugins/Developer/ProjectLauncher/Source/ProjectLauncher/Public/ProfileTree/LaunchProfileTreeData.h
--- a/Engine/Plugins/Developer/ProjectLauncher/Source/ProjectLauncher/Public/ProfileTree/LaunchProfileTreeData.h
+++ b/Engine/Plugins/Developer/ProjectLauncher/Source/ProjectLauncher/Public/ProfileTree/LaunchProfileTreeData.h
@@ -95,8 +95,9 @@ namespace ProjectLauncher
 	class FLaunchProfileTreeData : public TSharedFromThis<FLaunchProfileTreeData>
 	{
 	public:
-		// construction
+		// construction/destruction
 		UE_API FLaunchProfileTreeData(ILauncherProfilePtr InProfile, TSharedRef<FModel> InModel, ILaunchProfileTreeBuilder* InTreeBuilder);
+		UE_API ~FLaunchProfileTreeData();
 
 		// node creation
 		UE_API FLaunchProfileTreeNode& AddHeading( FText InName );
