Project Launcher extensions are notified when a property in the profile has changed

#rnx

diff --git a/Engine/Plugins/Developer/ProjectLauncher/Source/ProjectLauncher/Private/ProfileTree/GenericProfileTreeBuilder.cpp b/Engine/Plugins/Developer/ProjectLauncher/Source/ProjectLauncher/Private/ProfileTree/GenericProfileTreeBuilder.cpp
--- a/Engine/Plugins/Developer/ProjectLauncher/Source/ProjectLauncher/Private/ProfileTree/GenericProfileTreeBuilder.cpp
+++ b/Engine/Plugins/Developer/ProjectLauncher/Source/ProjectLauncher/Private/ProfileTree/GenericProfileTreeBuilder.cpp
@@ -568,7 +568,7 @@ namespace ProjectLauncher
 			Model->GetProfileManager()->SaveJSONProfile(Profile);
 		}
 
-		TreeData->RequestTreeRefresh();
+		TreeData->OnPropertyChanged();
 	}
 
 
diff --git a/Engine/Plugins/Developer/ProjectLauncher/Source/ProjectLauncher/Private/ProfileTree/LaunchProfileTreeData.cpp b/Engine/Plugins/Developer/ProjectLauncher/Source/ProjectLauncher/Private/ProfileTree/LaunchProfileTreeData.cpp
--- a/Engine/Plugins/Developer/ProjectLauncher/Source/ProjectLauncher/Private/ProfileTree/LaunchProfileTreeData.cpp
+++ b/Engine/Plugins/Developer/ProjectLauncher/Source/ProjectLauncher/Private/ProfileTree/LaunchProfileTreeData.cpp
@@ -80,6 +80,16 @@ namespace ProjectLauncher
 		bRequestTreeRefresh = true;
 	}
 
+	void FLaunchProfileTreeData::OnPropertyChanged()
+	{
+		for (TSharedPtr<FLaunchExtensionInstance> ExtensionInstance : ExtensionInstances)
+		{
+			ExtensionInstance->OnPropertyChanged();
+		}
+
+		RequestTreeRefresh();
+	}
+
 
 
 	FLaunchProfileTreeNode::FLaunchProfileTreeNode( FLaunchProfileTreeData* InTreeData ) 
diff --git a/Engine/Plugins/Developer/ProjectLauncher/Source/ProjectLauncher/Public/Extension/LaunchExtension.h b/Engine/Plugins/Developer/ProjectLauncher/Source/ProjectLauncher/Public/Extension/LaunchExtension.h
--- a/Engine/Plugins/Developer/ProjectLauncher/Source/ProjectLauncher/Public/Extension/LaunchExtension.h
+++ b/Engine/Plugins/Developer/ProjectLauncher/Source/ProjectLauncher/Public/Extension/LaunchExtension.h
@@ -76,6 +76,10 @@ namespace ProjectLauncher
 		 */
 		virtual void CustomizeLaunchCommandLine( FString& InOutCommandLine ) {};
 
+		/**
+		 * Notification callback when a property has been changed in our profile
+		 */
+		virtual void OnPropertyChanged() {};
 
 		/**
 		 * Determine if this extension also provides items in the bespoke extensions menu
diff --git a/Engine/Plugins/Developer/ProjectLauncher/Source/ProjectLauncher/Public/ProfileTree/LaunchProfileTreeData.h b/Engine/Plugins/Developer/ProjectLauncher/Source/ProjectLauncher/Public/ProfileTree/LaunchProfileTreeData.h
--- a/Engine/Plugins/Developer/ProjectLauncher/Source/ProjectLauncher/Public/ProfileTree/LaunchProfileTreeData.h
+++ b/Engine/Plugins/Developer/ProjectLauncher/Source/ProjectLauncher/Public/ProfileTree/LaunchProfileTreeData.h
@@ -103,6 +103,7 @@ namespace ProjectLauncher
 		UE_API FLaunchProfileTreeNode& AddHeading( FText InName );
 		UE_API void CreateExtensionsUI();
 
+		UE_API void OnPropertyChanged();
 		UE_API void RequestTreeRefresh();
 
 		// member variables
