Legacy Project Launcher tab is closed by default in UnrealFrontend.

#rnx

diff --git a/Engine/Plugins/Developer/ProjectLauncher/Source/ProjectLauncher/Private/Widgets/Profiles/SCustomLaunchProfilesPanel.cpp b/Engine/Plugins/Developer/ProjectLauncher/Source/ProjectLauncher/Private/Widgets/Profiles/SCustomLaunchProfilesPanel.cpp
--- a/Engine/Plugins/Developer/ProjectLauncher/Source/ProjectLauncher/Private/Widgets/Profiles/SCustomLaunchProfilesPanel.cpp
+++ b/Engine/Plugins/Developer/ProjectLauncher/Source/ProjectLauncher/Private/Widgets/Profiles/SCustomLaunchProfilesPanel.cpp
@@ -14,6 +14,7 @@
 #include "Widgets/Input/SButton.h"
 #include "Widgets/Input/SComboButton.h"
 #include "Widgets/Input/SCheckBox.h"
+#include "Widgets/Input/SHyperlink.h"
 #include "Widgets/Shared/SCustomLaunchProjectCombo.h"
 #include "Widgets/Shared/SCustomLaunchBuildTargetCombo.h"
 #include "Widgets/Profiles/SCustomLaunchCustomProfileSelector.h"
@@ -182,7 +183,26 @@ TSharedRef<SWidget> SCustomLaunchProfilesPanel::CreateProfilesPanel()
 				.SelectedProject(this, &SCustomLaunchProfilesPanel::GetDefaultProjectPath)
 				.Visibility(ProjectLauncher::bUseFriendlyBuildTargetSelection ? EVisibility::Collapsed : EVisibility::Visible)
 			]
-		];		
+
+			// Spacing
+			+SHorizontalBox::Slot()
+			.FillWidth(1)
+			[
+				SNew(SSpacer)
+			]
+
+			// Link to legacy Project Launcher
+			+SHorizontalBox::Slot()
+			.AutoWidth()
+			.VAlign(VAlign_Center)
+			.Padding(4,0)
+			[
+				SNew(SHyperlink)
+				.OnNavigate_Lambda([]() { FGlobalTabmanager::Get()->TryInvokeTab(FTabId("LegacyProjectLauncher")); } )
+				.Text(LOCTEXT("OpenLegacyProjectLauncher","Open Legacy Project Launcher"))
+				.ToolTipText(LOCTEXT("OpenLegacyProjectLauncherTip", "Opens the legacy Project Launcher tab"))
+			]
+		];
 	}
 
 	// main profile banner
diff --git a/Engine/Source/Programs/UnrealFrontend/Private/Commands/UserInterfaceCommand.cpp b/Engine/Source/Programs/UnrealFrontend/Private/Commands/UserInterfaceCommand.cpp
--- a/Engine/Source/Programs/UnrealFrontend/Private/Commands/UserInterfaceCommand.cpp
+++ b/Engine/Source/Programs/UnrealFrontend/Private/Commands/UserInterfaceCommand.cpp
@@ -167,7 +167,7 @@ void FUserInterfaceCommand::InitializeSlateApplication( const FString& LayoutIni
 						->AddTab(FName("MessagingDebugger"), ETabState::ClosedTab)
 						->AddTab(FName("SessionFrontend"), ETabState::OpenedTab)
 						->AddTab(FName("ProjectLauncher"), ETabState::OpenedTab)
-						->AddTab(FName("LegacyProjectLauncher"), ETabState::OpenedTab)
+						->AddTab(FName("LegacyProjectLauncher"), ETabState::ClosedTab)
 				)
 		)
 		->AddArea
