speculative crash fix for UFE startup crash on some profiles

#rnx

diff --git a/Engine/Plugins/Developer/ProjectLauncher/Source/ProjectLauncher/Private/Extension/LaunchExtension.cpp b/Engine/Plugins/Developer/ProjectLauncher/Source/ProjectLauncher/Private/Extension/LaunchExtension.cpp
--- a/Engine/Plugins/Developer/ProjectLauncher/Source/ProjectLauncher/Private/Extension/LaunchExtension.cpp
+++ b/Engine/Plugins/Developer/ProjectLauncher/Source/ProjectLauncher/Private/Extension/LaunchExtension.cpp
@@ -391,7 +391,13 @@ namespace ProjectLauncher
 
 	FName FLaunchExtensionInstance::GetLaunchPlatformName() const
 	{
-		return GetModel()->GetPlatformInfo(GetProfile())->IniPlatformName;
+		const PlatformInfo::FTargetPlatformInfo* Info = GetModel()->GetPlatformInfo(GetProfile());
+		if (Info != nullptr)
+		{
+			return Info->IniPlatformName;
+		}
+
+		return NAME_None;
 	}
 
 	FString FLaunchExtensionInstance::GetConfigString( EConfig Config, const TCHAR* Name, const TCHAR* DefaultValue ) const
