fix for assert  when re-selecting a profile

#rnx

diff --git a/Engine/Plugins/Developer/ProjectLauncher/Source/ProjectLauncher/Private/Extension/LaunchExtension.cpp b/Engine/Plugins/Developer/ProjectLauncher/Source/ProjectLauncher/Private/Extension/LaunchExtension.cpp
--- a/Engine/Plugins/Developer/ProjectLauncher/Source/ProjectLauncher/Private/Extension/LaunchExtension.cpp
+++ b/Engine/Plugins/Developer/ProjectLauncher/Source/ProjectLauncher/Private/Extension/LaunchExtension.cpp
@@ -87,7 +87,11 @@ namespace ProjectLauncher
 	TArray<TSharedPtr<FLaunchExtensionInstance>> FLaunchExtension::CreateExtensionInstancesForProfile( ILauncherProfileRef InProfile, TSharedRef<FModel> InModel )
 	{
 		// make sure the extensions for this profile haven't been created already
-		check(!GExtensionInstances.Contains(InProfile.ToSharedPtr()));
+		const TArray<TSharedPtr<FLaunchExtensionInstance>>* ExistingInstances = GExtensionInstances.Find(InProfile.ToSharedPtr());
+		if (ExistingInstances != nullptr)
+		{
+			return (*ExistingInstances);
+		}
 
 		// take a snapshot of the extension list
 		TArray<TSharedPtr<FLaunchExtension>> Extensions;
