Improve profile validation failure display + add custom validation options.

#rnx

diff --git a/Engine/Plugins/Developer/ProjectLauncher/Source/ProjectLauncher/Private/Extension/LaunchExtension.cpp b/Engine/Plugins/Developer/ProjectLauncher/Source/ProjectLauncher/Private/Extension/LaunchExtension.cpp
--- a/Engine/Plugins/Developer/ProjectLauncher/Source/ProjectLauncher/Private/Extension/LaunchExtension.cpp
+++ b/Engine/Plugins/Developer/ProjectLauncher/Source/ProjectLauncher/Private/Extension/LaunchExtension.cpp
@@ -185,10 +185,12 @@ namespace ProjectLauncher
 		, Model(InArgs.Model)
 		, Extension(InArgs.Extension)
 	{
+		Profile->OnCustomValidation().AddRaw( this, &FLaunchExtensionInstance::OnValidateProfile);
 	}
 
 	FLaunchExtensionInstance::~FLaunchExtensionInstance()
 	{
+		Profile->OnCustomValidation().RemoveAll(this);
 	}
 
 	bool FLaunchExtensionInstance::GetExtensionVariables( TArray<FString>& OutVariables ) const
diff --git a/Engine/Plugins/Developer/ProjectLauncher/Source/ProjectLauncher/Private/Model/ProjectLauncherModel.cpp b/Engine/Plugins/Developer/ProjectLauncher/Source/ProjectLauncher/Private/Model/ProjectLauncherModel.cpp
--- a/Engine/Plugins/Developer/ProjectLauncher/Source/ProjectLauncher/Private/Model/ProjectLauncherModel.cpp
+++ b/Engine/Plugins/Developer/ProjectLauncher/Source/ProjectLauncher/Private/Model/ProjectLauncherModel.cpp
@@ -962,30 +962,63 @@ namespace ProjectLauncher
 
 
 
-	FText GetProfileLaunchErrorMessage(ILauncherProfilePtr Profile)
+	FText GetProfileLaunchErrorMessage(ILauncherProfilePtr Profile, bool bWithAnnotations)
 	{
 		if (!Profile.IsValid())
 		{
 			return LOCTEXT("LaunchErrNoProfileTip", "There is no profile selected");
 		}
 
-		if (Profile->IsValidForLaunch())
+		FTextBuilder MsgTextBuilder;
+
+		TArray<FString> CustomErrors = Profile->GetAllCustomErrors();
+		if (CustomErrors.Num() > 0 || Profile->HasValidationError())
 		{
-			return LOCTEXT("LaunchProfileTip", "Launch this profile now");
+			if (bWithAnnotations)
+			{
+				MsgTextBuilder.AppendLine(LOCTEXT("LaunchErrValidation", "There are validation errors with this profile. Please fix them before launching:"));
+				MsgTextBuilder.Indent();
+			}
+			for (int i = 0; i < (int)ELauncherProfileValidationErrors::Count; i++)
+			{
+				ELauncherProfileValidationErrors::Type Error = (ELauncherProfileValidationErrors::Type)i;
+				if (Profile->HasValidationError(Error))
+				{
+					MsgTextBuilder.AppendLine(LexToStringLocalized(Error));
+				}
+			}
+			for (const FString& CustomError : CustomErrors)
+			{
+				MsgTextBuilder.AppendLine(Profile->GetCustomErrorText(CustomError));
+			}
+			if (bWithAnnotations)
+			{
+				MsgTextBuilder.Unindent();
+			}
+		}
+		else if (bWithAnnotations)
+		{
+			MsgTextBuilder.AppendLine(LOCTEXT("LaunchProfileTip", "Launch this profile now"));
 		}
 
-		FTextBuilder MsgTextBuilder;
-		MsgTextBuilder.AppendLine(LOCTEXT("LaunchErrValidation", "There are validation errors with this profile. Please fix them before launching:"));
-		MsgTextBuilder.Indent();
-		for (int i = 0; i < (int)ELauncherProfileValidationErrors::Count; i++)
+		TArray<FString> CustomWarnings = Profile->GetAllCustomWarnings();
+		if (CustomWarnings.Num() > 0)
 		{
-			ELauncherProfileValidationErrors::Type Error = (ELauncherProfileValidationErrors::Type)i;
-			if (Profile->HasValidationError(Error))
+			if (bWithAnnotations)
+			{
+				MsgTextBuilder.AppendLine();
+				MsgTextBuilder.AppendLine(LOCTEXT("LaunchWarnValidation", "There are validation warnings with this profile but these will not prevent launching:"));
+				MsgTextBuilder.Indent();
+			}
+			for (const FString& CustomWarning : CustomWarnings)
+			{
+				MsgTextBuilder.AppendLine(Profile->GetCustomWarningText(CustomWarning));
+			}
+			if (bWithAnnotations)
 			{
-				MsgTextBuilder.AppendLine(LexToStringLocalized(Error));
+				MsgTextBuilder.Unindent();
 			}
 		}
-		MsgTextBuilder.Unindent();
 
 		return MsgTextBuilder.ToText();
 	}
diff --git a/Engine/Plugins/Developer/ProjectLauncher/Source/ProjectLauncher/Private/ProfileTree/GenericProfileTreeBuilder.cpp b/Engine/Plugins/Developer/ProjectLauncher/Source/ProjectLauncher/Private/ProfileTree/GenericProfileTreeBuilder.cpp
--- a/Engine/Plugins/Developer/ProjectLauncher/Source/ProjectLauncher/Private/ProfileTree/GenericProfileTreeBuilder.cpp
+++ b/Engine/Plugins/Developer/ProjectLauncher/Source/ProjectLauncher/Private/ProfileTree/GenericProfileTreeBuilder.cpp
@@ -112,6 +112,9 @@ namespace ProjectLauncher
 			// - the first entry is the name of the property on the left-hand side
 			// - the seecond parameter is the widget itself that appears on the right-hand side
 			HeadingNode.AddWidget( LOCTEXT("ProjectLabel", "Project"), 
+				{
+					.Validation = FValidation( {ELauncherProfileValidationErrors::NoProjectSelected} ),
+				},
 				SNew(SCustomLaunchProjectCombo)
 				.OnSelectionChanged(this, &FGenericProfileTreeBuilder::SetProjectName)
 				.SelectedProject(this, &FGenericProfileTreeBuilder::GetProjectPath)
@@ -123,6 +126,9 @@ namespace ProjectLauncher
 		else
 		{
 			HeadingNode.AddWidget( LOCTEXT("ProjectLabel", "Project"), 
+				{
+					.Validation = FValidation( {ELauncherProfileValidationErrors::NoProjectSelected} ),
+				},
 				SNew(SCustomLaunchProjectCombo)
 				.OnSelectionChanged(this, &FGenericProfileTreeBuilder::SetProjectName)
 				.SelectedProject(this, &FGenericProfileTreeBuilder::GetProjectPath)
@@ -145,6 +151,7 @@ namespace ProjectLauncher
 				{
 					.IsDefault = [this]()		{ return !Profile->HasBuildTargetSpecified() || Profile->GetBuildTarget().IsEmpty(); },
 					.SetToDefault = [this]()	{ SetBuildTarget(FString()); },
+					.Validation = FValidation( {ELauncherProfileValidationErrors::BuildTargetCookVariantMismatch, ELauncherProfileValidationErrors::BuildTargetIsRequired, ELauncherProfileValidationErrors::CookOnTheFlyDoesntSupportServer} ),
 				},
 				SNew(SCustomLaunchBuildTargetCombo, Model)
 				.OnSelectionChanged(this, &FGenericProfileTreeBuilder::SetBuildTarget)
@@ -160,6 +167,7 @@ namespace ProjectLauncher
 					.IsDefault = [this]()		{ return !Profile->HasBuildTargetSpecified() || Profile->GetBuildTarget().IsEmpty(); },
 					.SetToDefault = [this]()	{ SetBuildTarget(FString()); },
 					.IsEnabled = [this]()		{ return Profile->HasProjectSpecified(); },
+					.Validation = FValidation( {ELauncherProfileValidationErrors::BuildTargetCookVariantMismatch, ELauncherProfileValidationErrors::BuildTargetIsRequired, ELauncherProfileValidationErrors::CookOnTheFlyDoesntSupportServer} ),
 				},
 				SNew(SCustomLaunchBuildTargetCombo, Model)
 				.OnSelectionChanged(this, &FGenericProfileTreeBuilder::SetBuildTarget)
@@ -174,6 +182,9 @@ namespace ProjectLauncher
 	void FGenericProfileTreeBuilder::AddPlatformProperty( FLaunchProfileTreeNode& HeadingNode )
 	{
 		HeadingNode.AddWidget( LOCTEXT("PlatformLabel", "Platform"), 
+			{
+				.Validation = FValidation({ELauncherProfileValidationErrors::NoPlatformSelected, ELauncherProfileValidationErrors::NoPlatformSDKInstalled})
+			},
 			SNew(SCustomLaunchPlatformCombo)
 			.SelectedPlatforms(this, &FGenericProfileTreeBuilder::GetSelectedPlatforms)
 			.OnSelectionChanged(this, &FGenericProfileTreeBuilder::SetSelectedPlatforms)
@@ -197,6 +208,9 @@ namespace ProjectLauncher
 		}
 
 		HeadingNode.AddWidget( LOCTEXT("ConfigurationLabel", "Configuration"), 
+			{
+				.Validation = FValidation({ELauncherProfileValidationErrors::NoBuildConfigurationSelected, ELauncherProfileValidationErrors::ShippingDoesntSupportCommandlineOptionsCantUseCookOnTheFly}),
+			},
 			SNew(SCustomLaunchLexToStringCombo<EBuildConfiguration>)
 			.OnSelectionChanged(this, &FGenericProfileTreeBuilder::SetBuildConfiguration)
 			.SelectedItem(this, &FGenericProfileTreeBuilder::GetBuildConfiguration)
@@ -207,6 +221,9 @@ namespace ProjectLauncher
 	void FGenericProfileTreeBuilder::AddContentSchemeProperty( FLaunchProfileTreeNode& HeadingNode )
 	{
 		HeadingNode.AddWidget( LOCTEXT("ContentSchemeLabel", "Content Scheme"), 
+			{
+				.Validation = FValidation( {ELauncherProfileValidationErrors::CookOnTheFlyDoesntSupportServer, ELauncherProfileValidationErrors::ShippingDoesntSupportCommandlineOptionsCantUseCookOnTheFly, ELauncherProfileValidationErrors::ZenPakStreamingRequiresDeployAndLaunch} ),
+			},
 			SNew(SCustomLaunchContentSchemeCombo)
 			.OnSelectionChanged(this, &FGenericProfileTreeBuilder::SetContentScheme)
 			.SelectedContentScheme_Lambda( [this]() { return ContentScheme; } )
@@ -237,6 +254,7 @@ namespace ProjectLauncher
 				.SetValue = [this](bool bValue)	{ Profile->SetUseIoStore(bValue); },
 				.GetDefaultValue = [this]()		{ return DefaultProfile->IsUsingIoStore(); },
 				.IsVisible = ForPak,
+				.Validation = FValidation({ELauncherProfileValidationErrors::IoStoreRequiresPakFiles}),
 			}
 		);
 	}
@@ -249,6 +267,7 @@ namespace ProjectLauncher
 				.SetValue = [this](bool bValue)	{ Profile->SetGenerateChunks(bValue); },
 				.GetDefaultValue = [this]()		{ return DefaultProfile->IsGeneratingChunks(); },
 				.IsVisible = ForPak,
+				.Validation = FValidation({ELauncherProfileValidationErrors::GeneratingChunksRequiresCookByTheBook, ELauncherProfileValidationErrors::GeneratingChunksRequiresUnrealPak}),
 			}
 		);
 	}
@@ -261,6 +280,7 @@ namespace ProjectLauncher
 				.SetValue = [this](bool bValue)	{ Profile->SetImportingZenSnapshot(bValue); },
 				.GetDefaultValue = [this]()		{ return DefaultProfile->IsImportingZenSnapshot(); },
 				.IsVisible = ForContent,
+				.Validation = FValidation({ELauncherProfileValidationErrors::CopyToDeviceRequiresCookByTheBook}),
 			}
 		);
 	}
@@ -273,6 +293,7 @@ namespace ProjectLauncher
 				.SetValue = [this](FString Value)	{ Profile->SetZenPakStreamingPath(Value); },
 				.GetDefaultValue = EmptyString,
 				.IsVisible = ForZenWS,
+				.Validation = FValidation({ELauncherProfileValidationErrors::ZenPakStreamingRequiresDeployAndLaunch}),
 			}
 		);
 	}
@@ -306,6 +327,7 @@ namespace ProjectLauncher
 				.IsDefault = [this]()		{ return GetIncrementalCookMode() == DefaultProfile->GetIncrementalCookMode(); },
 				.SetToDefault = [this]()	{ SetIncrementalCookMode(DefaultProfile->GetIncrementalCookMode()); },
 				.IsVisible = ForEnabledCooked,
+				.Validation = FValidation({ELauncherProfileValidationErrors::UnversionedAndIncremental}),
 			},
 			SNew(SCustomLaunchCombo<ELauncherProfileIncrementalCookMode::Type>)
 			.OnSelectionChanged(this, &FGenericProfileTreeBuilder::SetIncrementalCookMode)
@@ -331,7 +353,8 @@ namespace ProjectLauncher
 				.GetValue = [this]()			{ return GetCook(); },
 				.SetValue = [this](bool bValue)	{ SetCook(bValue); },
 				.GetDefaultValue = [this]()		{ return GetCook(DefaultProfile); },
-				.IsVisible = ForCooked
+				.IsVisible = ForCooked,
+				.Validation = FValidation({ELauncherProfileValidationErrors::CopyToDeviceRequiresCookByTheBook, ELauncherProfileValidationErrors::GeneratingPatchesCanOnlyRunFromByTheBookCookMode, ELauncherProfileValidationErrors::GeneratingChunksRequiresCookByTheBook}),
 			}
 		);
 	}
@@ -381,7 +404,7 @@ namespace ProjectLauncher
 		);
 	}
 
-	void FGenericProfileTreeBuilder::AddBuidUATProperty( FLaunchProfileTreeNode& HeadingNode )
+	void FGenericProfileTreeBuilder::AddBuildUATProperty( FLaunchProfileTreeNode& HeadingNode )
 	{
 		HeadingNode.AddBoolean( LOCTEXT("BuildUATLabel", "Build UAT"),
 			{
@@ -420,6 +443,7 @@ namespace ProjectLauncher
 				.SetValue = [this](FString Value)	{ Profile->SetPackageDirectory(Value); },
 				.GetDefaultValue = EmptyString,
 				.IsVisible = ForCooked,
+				.Validation = FValidation({ELauncherProfileValidationErrors::NoPackageDirectorySpecified}),
 			}
 		);
 	}
@@ -445,7 +469,8 @@ namespace ProjectLauncher
 				.SetValue = [this](FString Value)	{ Profile->SetArchiveDirectory(Value); },
 				.GetDefaultValue = EmptyString,
 				.IsVisible = ForCooked,				
-				.IsEnabled = [this]()				{ return Profile->IsArchiving(); }
+				.IsEnabled = [this]()				{ return Profile->IsArchiving(); },
+				.Validation = FValidation({ELauncherProfileValidationErrors::NoArchiveDirectorySpecified}),
 			}
 		);
 	}
@@ -459,7 +484,8 @@ namespace ProjectLauncher
 				.GetValue = [this]()			{ return GetDeployToDevice(); },
 				.SetValue = [this](bool bValue)	{ SetDeployToDevice(bValue); },
 				.GetDefaultValue = [this]()		{ return GetDeployToDevice(DefaultProfile); },
-				.IsVisible = [this]()			{ return (ContentScheme != EContentScheme::CookOnTheFly); }
+				.IsVisible = [this]()			{ return (ContentScheme != EContentScheme::CookOnTheFly); },
+				.Validation = FValidation({ELauncherProfileValidationErrors::CopyToDeviceRequiresCookByTheBook, ELauncherProfileValidationErrors::CopyToDeviceRequiresNoPackaging, ELauncherProfileValidationErrors::ZenPakStreamingRequiresDeployAndLaunch, ELauncherProfileValidationErrors::DeployedDeviceGroupRequired }),
 			}
 		);
 	}
@@ -481,6 +507,7 @@ namespace ProjectLauncher
 		HeadingNode.AddWidget( LOCTEXT("TargetDeviceLabel", "Target Device"), 
 			{
 				.IsEnabled = ForDeployment,
+				.Validation = FValidation( {ELauncherProfileValidationErrors::BuildTargetCookVariantMismatch, ELauncherProfileValidationErrors::LaunchDeviceIsUnauthorized} ),
 			},
 			CreateDeployDeviceWidget()
 		);
@@ -494,7 +521,8 @@ namespace ProjectLauncher
 				.GetValue = [this]()			{ return GetIsRunning(); },
 				.SetValue = [this](bool bValue)	{ SetIsRunning(bValue); },
 				.GetDefaultValue = [this]()		{ return GetIsRunning(DefaultProfile); },
-				.IsVisible = [this]()           { return Profile->GetAutomatedTests().Num() == 0; } // run is implied when there are automated tests
+				.IsVisible = [this]()           { return Profile->GetAutomatedTests().Num() == 0; }, // run is implied when there are automated tests
+				.Validation = FValidation({ELauncherProfileValidationErrors::ZenPakStreamingRequiresDeployAndLaunch, ELauncherProfileValidationErrors::NoLaunchRoleDeviceAssigned}),
 			}
 		);
 	}
@@ -508,6 +536,7 @@ namespace ProjectLauncher
 				.IsDefault = [this]()		{ return Profile->GetDefaultLaunchRole()->GetInitialMap().IsEmpty(); },
 				.SetToDefault = [this]()	{ Profile->GetDefaultLaunchRole()->SetInitialMap(FString()); InitalMapCombo->SetSelectedItem(GetInitialMap()); },
 				.IsEnabled = ForRun,
+				.Validation = FValidation({ELauncherProfileValidationErrors::InitialMapNotAvailable}),
 			},
 			SAssignNew(InitalMapCombo, SSearchableComboBox)
 			.OptionsSource(&CachedStartupMaps)
@@ -531,7 +560,8 @@ namespace ProjectLauncher
 				.GetValue = [this]()				{ return GetCommandLine(); },
 				.SetValue = [this](FString Value)	{ SetCommandLine(Value); },
 				.GetDefaultValue = EmptyString,
-				.IsEnabled = ForRun
+				.IsEnabled = ForRun,
+				.Validation = FValidation({ELauncherProfileValidationErrors::MalformedLaunchCommandLine}),
 			}
 		);
 	}
diff --git a/Engine/Plugins/Developer/ProjectLauncher/Source/ProjectLauncher/Private/ProfileTree/LaunchProfileTreeData.cpp b/Engine/Plugins/Developer/ProjectLauncher/Source/ProjectLauncher/Private/ProfileTree/LaunchProfileTreeData.cpp
--- a/Engine/Plugins/Developer/ProjectLauncher/Source/ProjectLauncher/Private/ProfileTree/LaunchProfileTreeData.cpp
+++ b/Engine/Plugins/Developer/ProjectLauncher/Source/ProjectLauncher/Private/ProfileTree/LaunchProfileTreeData.cpp
@@ -25,6 +25,92 @@
 
 namespace ProjectLauncher
 {
+	FValidation::FValidation(std::initializer_list<ELauncherProfileValidationErrors::Type> InAssociatedErrors )
+		: AssociatedErrors(InAssociatedErrors)
+	{
+		bIsSet = true;
+	}
+
+	FValidation::FValidation(std::initializer_list<ELauncherProfileValidationErrors::Type> InAssociatedErrors, std::initializer_list<FString> InAssociatedCustomErrors )
+		: AssociatedErrors(InAssociatedErrors)
+		, AssociatedCustomErrors(InAssociatedCustomErrors)
+	{
+		bIsSet = true;
+	}
+
+	FValidation::FValidation( std::initializer_list<FString> InAssociatedCustomErrors )
+		: AssociatedCustomErrors(InAssociatedCustomErrors)
+	{
+		bIsSet = true;
+	}
+
+
+	bool FValidation::IsSet() const
+	{
+		return bIsSet;
+	}
+
+	bool FValidation::HasError( ILauncherProfileRef Profile ) const
+	{
+		if (Profile->HasValidationError())
+		{
+			for (ELauncherProfileValidationErrors::Type Error : AssociatedErrors)
+			{
+				if (Profile->HasValidationError(Error))
+				{
+					return true;
+				}
+			}
+		}
+
+		for (const FString& CustomError : AssociatedCustomErrors)
+		{
+			if (Profile->HasCustomError(CustomError))
+			{
+				return true;
+			}
+		}
+
+		for (const FString& CustomError : AssociatedCustomErrors)
+		{
+			if (Profile->HasCustomWarning(CustomError))
+			{
+				return true;
+			}
+		}
+
+		return false;
+	}
+
+	void FValidation::GetErrorText( ILauncherProfileRef Profile, TArray<FString>& ErrorLines ) const
+	{
+		for (ELauncherProfileValidationErrors::Type Error : AssociatedErrors)
+		{
+			if (Profile->HasValidationError(Error))
+			{
+				ErrorLines.AddUnique(LexToStringLocalized(Error));
+			}
+		}
+
+		for (const FString& CustomError : AssociatedCustomErrors)
+		{
+			if (Profile->HasCustomError(CustomError))
+			{
+				ErrorLines.AddUnique(Profile->GetCustomErrorText(CustomError).ToString());
+			}
+		}
+
+		for (const FString& CustomWarning : AssociatedCustomErrors)
+		{
+			if (Profile->HasCustomWarning(CustomWarning))
+			{
+				ErrorLines.AddUnique(Profile->GetCustomWarningText(CustomWarning).ToString());
+			}
+		}
+	}
+
+
+
 	FLaunchProfileTreeData::FLaunchProfileTreeData(ILauncherProfilePtr InProfile, TSharedRef<ProjectLauncher::FModel> InModel, ILaunchProfileTreeBuilder* InTreeBuilder) 
 		: Profile(InProfile)
 		, Model(InModel)
@@ -78,6 +164,7 @@ namespace ProjectLauncher
 	void FLaunchProfileTreeData::RequestTreeRefresh()
 	{
 		bRequestTreeRefresh = true;
+		Profile->RefreshCustomWarningsAndErrors();
 	}
 
 	void FLaunchProfileTreeData::OnPropertyChanged()
@@ -124,6 +211,7 @@ namespace ProjectLauncher
 		{
 			.IsVisible = BooleanCallbacks.IsVisible,
 			.IsEnabled = BooleanCallbacks.IsEnabled,
+			.Validation = BooleanCallbacks.Validation,
 		};
 		if (BooleanCallbacks.GetDefaultValue != nullptr)
 		{
@@ -167,6 +255,7 @@ namespace ProjectLauncher
 		{
 			.IsVisible = StringCallbacks.IsVisible,
 			.IsEnabled = StringCallbacks.IsEnabled,
+			.Validation = StringCallbacks.Validation,
 		};
 		if (StringCallbacks.GetDefaultValue != nullptr)
 		{
@@ -214,6 +303,7 @@ namespace ProjectLauncher
 		{
 			.IsVisible = StringCallbacks.IsVisible,
 			.IsEnabled = StringCallbacks.IsEnabled,
+			.Validation = StringCallbacks.Validation,
 		};
 		if (StringCallbacks.GetDefaultValue != nullptr)
 		{
@@ -319,6 +409,7 @@ namespace ProjectLauncher
 		{
 			.IsVisible = StringCallbacks.IsVisible,
 			.IsEnabled = StringCallbacks.IsEnabled,
+			.Validation = StringCallbacks.Validation,
 		};
 		if (StringCallbacks.GetDefaultValue != nullptr)
 		{
diff --git a/Engine/Plugins/Developer/ProjectLauncher/Source/ProjectLauncher/Private/Styling/ProjectLauncherStyle.cpp b/Engine/Plugins/Developer/ProjectLauncher/Source/ProjectLauncher/Private/Styling/ProjectLauncherStyle.cpp
--- a/Engine/Plugins/Developer/ProjectLauncher/Source/ProjectLauncher/Private/Styling/ProjectLauncherStyle.cpp
+++ b/Engine/Plugins/Developer/ProjectLauncher/Source/ProjectLauncher/Private/Styling/ProjectLauncherStyle.cpp
@@ -75,7 +75,8 @@ void FProjectLauncherStyle::Initialize()
 	StyleSet->Set("Icons.DeviceManager", new CORE_IMAGE_BRUSH_SVG( "../Editor/Slate/Starship/Common/DeviceManager", Icon16x16) );
 	StyleSet->Set("Icons.Plugins", new CORE_IMAGE_BRUSH_SVG("../Editor/Slate/Starship/Common/Plugins", Icon16x16));
 
-	StyleSet->Set("Icons.Asset", new CORE_IMAGE_BRUSH( "Icons/doc_16x", Icon16x16 ) );
+	StyleSet->Set("Icons.Asset", new CORE_IMAGE_BRUSH( "./Editor/Slate/Icons/doc_16x", Icon16x16 ) );
+	StyleSet->Set("Icons.WarningWithColor.Small", new CORE_IMAGE_BRUSH_SVG("../Editor/Slate/Starship/Common/AlertTriangleSolid", Icon16x16, FStyleColors::Warning));
 
 	StyleSet->Set( "PathPickerButton", new CORE_IMAGE_BRUSH( "Icons/ellipsis_12x", Icon12x12 ) );
 
diff --git a/Engine/Plugins/Developer/ProjectLauncher/Source/ProjectLauncher/Private/Widgets/Profiles/SCustomLaunchCustomProfileEditor.cpp b/Engine/Plugins/Developer/ProjectLauncher/Source/ProjectLauncher/Private/Widgets/Profiles/SCustomLaunchCustomProfileEditor.cpp
--- a/Engine/Plugins/Developer/ProjectLauncher/Source/ProjectLauncher/Private/Widgets/Profiles/SCustomLaunchCustomProfileEditor.cpp
+++ b/Engine/Plugins/Developer/ProjectLauncher/Source/ProjectLauncher/Private/Widgets/Profiles/SCustomLaunchCustomProfileEditor.cpp
@@ -65,6 +65,40 @@ public:
 			return false;
 		};
 
+		auto GetErrorVisibility = [this, TreeNode]()
+		{
+			if (!IsItemExpanded())
+			{
+				for (ProjectLauncher::FLaunchProfileTreeNodePtr ChildNode : TreeNode->Children)
+				{
+					if (ChildNode->Callbacks.Validation.HasError( TreeNode->GetTreeData()->Profile.ToSharedRef() ))
+					{
+						return EVisibility::Visible;
+					}
+				}
+			}
+			
+			return EVisibility::Collapsed;
+		};
+
+		auto GetErrorToolTipText = [TreeNode]()
+		{
+			TArray<FString> ErrorLines;
+			for (ProjectLauncher::FLaunchProfileTreeNodePtr ChildNode : TreeNode->Children)
+			{
+				ChildNode->Callbacks.Validation.GetErrorText( TreeNode->GetTreeData()->Profile.ToSharedRef(), ErrorLines );
+			}
+
+			FTextBuilder TextBuilder;
+			for (const FString& ErrorLine : ErrorLines)
+			{
+				TextBuilder.AppendLine(FText::FromString(ErrorLine));
+			}
+
+			return TextBuilder.ToText();
+		};
+
+
 
 		this->ChildSlot
 		[
@@ -95,12 +129,24 @@ public:
 						+SHorizontalBox::Slot()
 						.VAlign(VAlign_Center)
 						.Padding(4, 0, 0, 0)
-						.FillWidth(1)
+						.AutoWidth()
 						[
 							SNew(STextBlock)
 							.Text(TreeNode->Name)
 							.Font(FCoreStyle::Get().GetFontStyle("SmallFontBold"))
 						]
+
+						+SHorizontalBox::Slot()
+						.HAlign(HAlign_Left)
+						.VAlign(VAlign_Center)
+						.Padding(2, 0)
+						.AutoWidth()
+						[
+							SNew(SImage)
+							.Image(FProjectLauncherStyle::Get().GetBrush("Icons.WarningWithColor.Small"))
+							.Visibility_Lambda(GetErrorVisibility)
+							.ToolTipText_Lambda(GetErrorToolTipText)
+						]
 					]
 				]
 			]
@@ -172,6 +218,30 @@ public:
 			;
 		}
 
+		auto GetErrorVisibility = [TreeNode]()
+		{
+			if (TreeNode->Callbacks.Validation.HasError( TreeNode->GetTreeData()->Profile.ToSharedRef() ))
+			{
+				return EVisibility::Visible;
+			}
+			
+			return EVisibility::Collapsed;
+		};
+
+		auto GetErrorToolTipText = [TreeNode]()
+		{
+			TArray<FString> ErrorLines;
+			TreeNode->Callbacks.Validation.GetErrorText( TreeNode->GetTreeData()->Profile.ToSharedRef(), ErrorLines );
+
+			FTextBuilder TextBuilder;
+			for (const FString& ErrorLine : ErrorLines)
+			{
+				TextBuilder.AppendLine(FText::FromString(ErrorLine));
+			}
+
+			return TextBuilder.ToText();
+		};
+
 		auto MakeSplitterSlot = [this]( TSharedRef<SWidget> SlotContent )
 		{
 			return SNew(SBorder)
@@ -193,6 +263,8 @@ public:
 			];
 		};
 
+		TSharedPtr<SHorizontalBox> PropertyNameBox;
+
 
 		this->ChildSlot
 		[
@@ -210,12 +282,20 @@ public:
 			[
 				MakeSplitterSlot(
 					SNew(SBox)
-					.Padding( FMargin(24, 4, 4, 4))
+					.Padding( FMargin(24, 1, 4, 0))
 					.VAlign(VAlign_Top)
 					[
-						SNew(STextBlock)
-						.Text(TreeNode->Name)
-						.Font(FCoreStyle::Get().GetFontStyle("SmallFont"))
+						SAssignNew(PropertyNameBox, SHorizontalBox)
+						+SHorizontalBox::Slot()
+						.AutoWidth()
+						.HAlign(HAlign_Left)
+						.VAlign(VAlign_Center)
+						.Padding(0,4)
+						[
+							SNew(STextBlock)
+							.Text(TreeNode->Name)
+							.Font(FCoreStyle::Get().GetFontStyle("SmallFont"))
+						]
 					]
 				)
 			]
@@ -254,6 +334,21 @@ public:
 			]
 		];
 
+		if (TreeNode->Callbacks.Validation.IsSet() )
+		{
+			PropertyNameBox->AddSlot()
+			.AutoWidth()
+			.HAlign(HAlign_Left)
+			.VAlign(VAlign_Center)
+			.Padding(2, 1, 0, 0)
+			[
+				SNew(SImage)
+				.Image(FProjectLauncherStyle::Get().GetBrush("Icons.WarningWithColor.Small"))
+				.Visibility_Lambda(GetErrorVisibility)
+				.ToolTipText_Lambda(GetErrorToolTipText)
+			];
+		}
+
 
 		SLaunchProfileTreeRow::ConstructInternal(
 			STableRow::FArguments()
@@ -277,15 +372,79 @@ void SCustomLaunchCustomProfileEditor::Construct(const FArguments& InArgs, const
 
 	TreeBuilder = ProjectLauncher::CreateTreeBuilder( nullptr, InModel );
 
+	auto GetErrorVisibility = [this]()
+	{
+		if (CurrentProfile.IsValid() && (!CurrentProfile->IsValidForLaunch() || CurrentProfile->GetAllCustomWarnings().Num() > 0))
+		{
+			return EVisibility::Visible;
+		}
+			
+		return EVisibility::Collapsed;
+	};
+
+	auto GetErrorText = [this]()
+	{
+		return ProjectLauncher::GetProfileLaunchErrorMessage(CurrentProfile);
+	};
+
+
 
 	ChildSlot
 	[
-		SAssignNew(TreeView, SLaunchProfileTreeView)
-		.TreeItemsSource(&TreeBuilder->GetProfileTree()->Nodes)
-		.SelectionMode( ESelectionMode::None )
-		.OnGenerateRow( this, &SCustomLaunchCustomProfileEditor::OnGenerateWidgetForTreeNode )
-		.OnGetChildren( this, &SCustomLaunchCustomProfileEditor::OnGetChildren)
-		.HandleDirectionalNavigation(false)
+		SNew(SBorder)
+		.BorderImage(FAppStyle::GetBrush("ChildWindow.Background"))
+		.Padding(2)
+		[
+			SNew(SVerticalBox)
+			+SVerticalBox::Slot()
+			.AutoHeight()
+			.VAlign(VAlign_Center)
+			.HAlign(HAlign_Fill)
+			.Padding(2)
+			[
+				SNew(SBorder)
+				.Visibility_Lambda(GetErrorVisibility)
+				.Padding(8)
+
+				.BorderImage(FAppStyle::GetBrush("ToolPanel.GroupBorder"))
+				[
+					SNew(SHorizontalBox)
+
+					+SHorizontalBox::Slot()
+					.AutoWidth()
+					.HAlign(HAlign_Left)
+					.VAlign(VAlign_Center)
+					.Padding(4)
+					[
+						SNew(SImage)
+						.Image(FAppStyle::Get().GetBrush("Icons.WarningWithColor.Large"))
+						.DesiredSizeOverride(FVector2D(24,24))
+					]
+
+					+SHorizontalBox::Slot()
+					.FillWidth(1)
+					.HAlign(HAlign_Fill)
+					.VAlign(VAlign_Center)
+					.Padding(4)
+					[
+						SNew(STextBlock)
+						.Text_Lambda(GetErrorText)
+					]
+				]
+			]
+
+			+SVerticalBox::Slot()
+			.FillHeight(1)
+			.Padding(2)
+			[
+				SAssignNew(TreeView, SLaunchProfileTreeView)
+				.TreeItemsSource(&TreeBuilder->GetProfileTree()->Nodes)
+				.SelectionMode( ESelectionMode::None )
+				.OnGenerateRow( this, &SCustomLaunchCustomProfileEditor::OnGenerateWidgetForTreeNode )
+				.OnGetChildren( this, &SCustomLaunchCustomProfileEditor::OnGetChildren)
+				.HandleDirectionalNavigation(false)
+			]
+		]
 	];
 
 }
diff --git a/Engine/Plugins/Developer/ProjectLauncher/Source/ProjectLauncher/Private/Widgets/Profiles/SCustomLaunchCustomProfileSelector.cpp b/Engine/Plugins/Developer/ProjectLauncher/Source/ProjectLauncher/Private/Widgets/Profiles/SCustomLaunchCustomProfileSelector.cpp
--- a/Engine/Plugins/Developer/ProjectLauncher/Source/ProjectLauncher/Private/Widgets/Profiles/SCustomLaunchCustomProfileSelector.cpp
+++ b/Engine/Plugins/Developer/ProjectLauncher/Source/ProjectLauncher/Private/Widgets/Profiles/SCustomLaunchCustomProfileSelector.cpp
@@ -114,6 +114,22 @@ TSharedRef<ITableRow> SCustomLaunchCustomProfileSelector::GenerateCustomProfileR
 {
 	bool bIsBasicLaunchProfile = Model->IsBasicLaunchProfile(Profile);
 
+	auto GetErrorVisibility = [Profile]()
+	{
+		if (Profile->HasValidationError())
+		{
+			return EVisibility::Visible;
+		}
+
+		return EVisibility::Collapsed;
+	};
+
+	auto GetErrorToolTipText = [Profile]()
+	{
+		return ProjectLauncher::GetProfileLaunchErrorMessage(Profile);
+	};
+
+
 	TSharedPtr<SVerticalBox> DetailsBox;
 	TSharedPtr<SVerticalBox> ButtonsBox;
 	TSharedPtr<SHorizontalBox> MainBox;
@@ -155,6 +171,9 @@ TSharedRef<ITableRow> SCustomLaunchCustomProfileSelector::GenerateCustomProfileR
 	// set the color attribute after the row itself has been created - this allows us to capture the row widget itself
 	CustomSelectionBorder->SetBorderBackgroundColor(TAttribute<FSlateColor>::CreateSP(this, &SCustomLaunchCustomProfileSelector::GetRowColor, RowContent) );
 	
+	TSharedPtr<SHorizontalBox> NameBox;
+
+
 	if (bIsBasicLaunchProfile)
 	{
 		// add the name
@@ -162,9 +181,15 @@ TSharedRef<ITableRow> SCustomLaunchCustomProfileSelector::GenerateCustomProfileR
 		.Padding(4)
 		.AutoHeight()
 		[
-			SNew(STextBlock)
-			.ColorAndOpacity(FColor::White)
-			.Text_Lambda( [Profile]() { return FText::FromString(Profile->GetName()); })
+			SAssignNew(NameBox, SHorizontalBox)
+
+			+SHorizontalBox::Slot()
+			.AutoWidth()
+			[
+				SNew(STextBlock)
+				.ColorAndOpacity(FColor::White)
+				.Text_Lambda( [Profile]() { return FText::FromString(Profile->GetName()); })
+			]
 		];
 
 		// add the quick device picker
@@ -216,10 +241,16 @@ TSharedRef<ITableRow> SCustomLaunchCustomProfileSelector::GenerateCustomProfileR
 		.Padding(4)
 		.AutoHeight()
 		[
-			SAssignNew(NameEditTextBox, SInlineEditableTextBlock)
-			.Text_Lambda( [Profile]() { return FText::FromString(Profile->GetName()); })
-			.ColorAndOpacity(FColor::White)
-			.OnTextCommitted( this, &SCustomLaunchCustomProfileSelector::SetProfileName, Profile )
+			SAssignNew(NameBox, SHorizontalBox)
+
+			+SHorizontalBox::Slot()
+			.AutoWidth()
+			[
+				SAssignNew(NameEditTextBox, SInlineEditableTextBlock)
+				.Text_Lambda( [Profile]() { return FText::FromString(Profile->GetName()); })
+				.ColorAndOpacity(FColor::White)
+				.OnTextCommitted( this, &SCustomLaunchCustomProfileSelector::SetProfileName, Profile )
+			]
 		];
 
 		// description
@@ -239,6 +270,21 @@ TSharedRef<ITableRow> SCustomLaunchCustomProfileSelector::GenerateCustomProfileR
 		DescriptionEditTextBoxes.Add(Profile, DescriptionEditTextBox);
 	}
 
+	// validation warning icon
+	NameBox->AddSlot()
+	.HAlign(HAlign_Left)
+	.VAlign(VAlign_Center)
+	.Padding(4, 0)
+	.AutoWidth()
+	[
+		SNew(SImage)
+		.Image(FProjectLauncherStyle::Get().GetBrush("Icons.WarningWithColor.Small"))
+		.Visibility_Lambda(GetErrorVisibility)
+		.ToolTipText_Lambda(GetErrorToolTipText)
+	];
+
+
+
 	// floating edit button
 	MainBox->AddSlot()
 	.VAlign(VAlign_Center)
diff --git a/Engine/Plugins/Developer/ProjectLauncher/Source/ProjectLauncher/Private/Widgets/Profiles/SCustomLaunchProfilesPanel.cpp b/Engine/Plugins/Developer/ProjectLauncher/Source/ProjectLauncher/Private/Widgets/Profiles/SCustomLaunchProfilesPanel.cpp
--- a/Engine/Plugins/Developer/ProjectLauncher/Source/ProjectLauncher/Private/Widgets/Profiles/SCustomLaunchProfilesPanel.cpp
+++ b/Engine/Plugins/Developer/ProjectLauncher/Source/ProjectLauncher/Private/Widgets/Profiles/SCustomLaunchProfilesPanel.cpp
@@ -666,7 +666,7 @@ bool SCustomLaunchProfilesPanel::IsLaunchButtonEnabled() const
 FText SCustomLaunchProfilesPanel::GetLaunchButtonToolTipText() const
 {
 	ILauncherProfilePtr Profile = Model->GetSelectedProfile();
-	return ProjectLauncher::GetProfileLaunchErrorMessage(Profile);
+	return ProjectLauncher::GetProfileLaunchErrorMessage(Profile, true);
 }
 
 const FSlateBrush* SCustomLaunchProfilesPanel::GetLaunchButtonImage() const
diff --git a/Engine/Plugins/Developer/ProjectLauncher/Source/ProjectLauncher/Public/Extension/LaunchExtension.h b/Engine/Plugins/Developer/ProjectLauncher/Source/ProjectLauncher/Public/Extension/LaunchExtension.h
--- a/Engine/Plugins/Developer/ProjectLauncher/Source/ProjectLauncher/Public/Extension/LaunchExtension.h
+++ b/Engine/Plugins/Developer/ProjectLauncher/Source/ProjectLauncher/Public/Extension/LaunchExtension.h
@@ -86,6 +86,11 @@ namespace ProjectLauncher
 		 */
 		virtual void OnPropertyChanged() {};
 
+		/**
+		 * Notification callback when the profile is validated
+		 */
+		virtual void OnValidateProfile() {};
+
 		/**
 		 * Determine if this extension also provides items in the bespoke extensions menu
 		 */
diff --git a/Engine/Plugins/Developer/ProjectLauncher/Source/ProjectLauncher/Public/Model/ProjectLauncherModel.h b/Engine/Plugins/Developer/ProjectLauncher/Source/ProjectLauncher/Public/Model/ProjectLauncherModel.h
--- a/Engine/Plugins/Developer/ProjectLauncher/Source/ProjectLauncher/Public/Model/ProjectLauncherModel.h
+++ b/Engine/Plugins/Developer/ProjectLauncher/Source/ProjectLauncher/Public/Model/ProjectLauncherModel.h
@@ -38,7 +38,7 @@ namespace ProjectLauncher
 	TArray<EContentScheme> GetAllContentSchemes();
 	FText GetContentSchemeDisplayName(EContentScheme ContentScheme);
 	FText GetContentSchemeToolTip(EContentScheme ContentScheme);
-	FText GetProfileLaunchErrorMessage(ILauncherProfilePtr Profile);
+	FText GetProfileLaunchErrorMessage(ILauncherProfilePtr Profile, bool bWithAnnotations = false );
 	const TCHAR* LexToString( const ProjectLauncher::EContentScheme& ContentScheme);
 	bool LexTryParseString( ProjectLauncher::EContentScheme& OutContentScheme, const TCHAR* String );
 
diff --git a/Engine/Plugins/Developer/ProjectLauncher/Source/ProjectLauncher/Public/ProfileTree/GenericProfileTreeBuilder.h b/Engine/Plugins/Developer/ProjectLauncher/Source/ProjectLauncher/Public/ProfileTree/GenericProfileTreeBuilder.h
--- a/Engine/Plugins/Developer/ProjectLauncher/Source/ProjectLauncher/Public/ProfileTree/GenericProfileTreeBuilder.h
+++ b/Engine/Plugins/Developer/ProjectLauncher/Source/ProjectLauncher/Public/ProfileTree/GenericProfileTreeBuilder.h
@@ -60,7 +60,7 @@ namespace ProjectLauncher
 		UE_API void AddAdditionalCookerOptionsProperty( FLaunchProfileTreeNode& HeadingNode );
 		UE_API void AddBuildProperty( FLaunchProfileTreeNode& HeadingNode );
 		UE_API void AddForceBuildProperty( FLaunchProfileTreeNode& HeadingNode );
-		UE_API void AddBuidUATProperty( FLaunchProfileTreeNode& HeadingNode );
+		UE_API void AddBuildUATProperty( FLaunchProfileTreeNode& HeadingNode );
 		UE_API void AddArchitectureProperty( FLaunchProfileTreeNode& HeadingNode );
 		UE_API void AddStagingDirectoryProperty( FLaunchProfileTreeNode& HeadingNode );
 		UE_API void AddArchiveBuildProperty( FLaunchProfileTreeNode& HeadingNode );
diff --git a/Engine/Plugins/Developer/ProjectLauncher/Source/ProjectLauncher/Public/ProfileTree/LaunchProfileTreeData.h b/Engine/Plugins/Developer/ProjectLauncher/Source/ProjectLauncher/Public/ProfileTree/LaunchProfileTreeData.h
--- a/Engine/Plugins/Developer/ProjectLauncher/Source/ProjectLauncher/Public/ProfileTree/LaunchProfileTreeData.h
+++ b/Engine/Plugins/Developer/ProjectLauncher/Source/ProjectLauncher/Public/ProfileTree/LaunchProfileTreeData.h
@@ -6,6 +6,7 @@
 #include "Templates/Function.h"
 #include "Templates/SharedPointer.h"
 #include "ILauncherProfile.h"
+#include <initializer_list>
 
 #define UE_API PROJECTLAUNCHER_API
 
@@ -26,6 +27,25 @@ namespace ProjectLauncher
 	typedef TSharedPtr<FLaunchProfileTreeNode> FLaunchProfileTreeNodePtr;
 	typedef TSharedRef<FLaunchProfileTreeNode> FLaunchProfileTreeNodeRef;
 
+	class FValidation
+	{
+	public:
+		FValidation() {}
+		FValidation( std::initializer_list<ELauncherProfileValidationErrors::Type> InAssociatedErrors );
+		FValidation( std::initializer_list<ELauncherProfileValidationErrors::Type> InAssociatedErrors, std::initializer_list<FString> InAssociatedCustomErrors );
+		FValidation( std::initializer_list<FString> InAssociatedCustomErrors );
+
+		bool IsSet() const;
+
+		bool HasError( ILauncherProfileRef Profile ) const;
+		void GetErrorText( ILauncherProfileRef Profile, TArray<FString>& ErrorLines ) const;
+
+	private:
+		bool bIsSet = false;
+		TArray<ELauncherProfileValidationErrors::Type> AssociatedErrors;
+		TArray<FString> AssociatedCustomErrors;
+	};
+
 
 	class FLaunchProfileTreeNode : public TSharedFromThis<FLaunchProfileTreeNode>
 	{
@@ -46,6 +66,7 @@ namespace ProjectLauncher
 			FFunction SetToDefault = nullptr;
 			FGetBool IsVisible = nullptr;
 			FGetBool IsEnabled = nullptr;
+			FValidation Validation;
 		};
 		UE_API FLaunchProfileTreeNode& AddWidget( FText InName, FCallbacks&& InWidgetCallbacks, TSharedRef<SWidget> InValueWidget );
 		UE_API FLaunchProfileTreeNode& AddWidget( FText InName, TSharedRef<SWidget> InValueWidget );
@@ -58,6 +79,7 @@ namespace ProjectLauncher
 			FGetBool GetDefaultValue = nullptr;
 			FGetBool IsVisible = nullptr;
 			FGetBool IsEnabled = nullptr;
+			FValidation Validation;
 		};
 		UE_API FLaunchProfileTreeNode& AddBoolean( FText InName, FBooleanCallbacks&& BooleanCallbacks );
 
@@ -68,7 +90,7 @@ namespace ProjectLauncher
 			FGetString GetDefaultValue = nullptr;
 			FGetBool IsVisible = nullptr;
 			FGetBool IsEnabled = nullptr;
-
+			FValidation Validation;
 		};
 		UE_API FLaunchProfileTreeNode& AddString( FText InName, FStringCallbacks&& StringCallbacks );
 		UE_API FLaunchProfileTreeNode& AddDirectoryString( FText InName, FStringCallbacks&& StringCallbacks );
diff --git a/Engine/Source/Developer/LauncherServices/Private/Profiles/LauncherProfile.cpp b/Engine/Source/Developer/LauncherServices/Private/Profiles/LauncherProfile.cpp
--- a/Engine/Source/Developer/LauncherServices/Private/Profiles/LauncherProfile.cpp
+++ b/Engine/Source/Developer/LauncherServices/Private/Profiles/LauncherProfile.cpp
@@ -10,7 +10,7 @@ FString LexToStringLocalized(ELauncherProfileValidationErrors::Type Value)
 	switch (Value)
 	{
 		case ELauncherProfileValidationErrors::CopyToDeviceRequiresCookByTheBook:
-			return LOCTEXT("CopyToDeviceRequiresCookByTheBookError", "Deployment by copying to device requires 'By The Book' cooking.").ToString();
+			return LOCTEXT("CopyToDeviceRequiresCookByTheBookError", "Deployment by copying to device requires cooking or zen snapshot import.").ToString();
 		case ELauncherProfileValidationErrors::CustomRolesNotSupportedYet:
 			return LOCTEXT("CustomRolesNotSupportedYet", "Custom launch roles are not supported yet.").ToString();
 		case ELauncherProfileValidationErrors::DeployedDeviceGroupRequired:
@@ -24,11 +24,11 @@ FString LexToStringLocalized(ELauncherProfileValidationErrors::Type Value)
 		case ELauncherProfileValidationErrors::NoBuildConfigurationSelected:
 			return LOCTEXT("NoBuildConfigurationSelectedError", "A Build Configuration must be selected.").ToString();
 		case ELauncherProfileValidationErrors::NoCookedCulturesSelected:
-			return LOCTEXT("NoCookedCulturesSelectedError", "At least one Culture must be selected when cooking by the book.").ToString();
+			return LOCTEXT("NoCookedCulturesSelectedError", "At least one Culture must be selected when cooking.").ToString();
 		case ELauncherProfileValidationErrors::NoLaunchRoleDeviceAssigned:
 			return LOCTEXT("NoLaunchRoleDeviceAssigned", "One or more launch roles do not have a device assigned.").ToString();
 		case ELauncherProfileValidationErrors::NoPlatformSelected:
-			return LOCTEXT("NoCookedPlatformSelectedError", "At least one Platform must be selected when cooking by the book.").ToString();
+			return LOCTEXT("NoCookedPlatformSelectedError", "At least one Platform must be selected.").ToString();
 		case ELauncherProfileValidationErrors::NoProjectSelected:
 			return LOCTEXT("NoBuildGameSelectedError", "A Project must be selected.").ToString();
 		case ELauncherProfileValidationErrors::NoPackageDirectorySpecified:
@@ -38,13 +38,13 @@ FString LexToStringLocalized(ELauncherProfileValidationErrors::Type Value)
 		case ELauncherProfileValidationErrors::UnversionedAndIncremental:
 			return LOCTEXT("UnversionedAndIncremental", "Unversioned build cannot be incremental.").ToString();
 		case ELauncherProfileValidationErrors::GeneratingPatchesCanOnlyRunFromByTheBookCookMode:
-			return LOCTEXT("GeneratingPatchesCanOnlyRunFromByTheBookCookMode", "Generating patch requires cook by the book mode.").ToString();
+			return LOCTEXT("GeneratingPatchesCanOnlyRunFromByTheBookCookMode", "Generating patch requires cooking or zen snapshot import.").ToString();
 		case ELauncherProfileValidationErrors::GeneratingMultiLevelPatchesRequiresGeneratePatch:
 			return LOCTEXT("GeneratingMultiLevelPatchesRequiresGeneratePatch", "Generating multilevel patch requires generating patch.").ToString();
 		case ELauncherProfileValidationErrors::StagingBaseReleasePaksWithoutABaseReleaseVersion:
 			return LOCTEXT("StagingBaseReleasePaksWithoutABaseReleaseVersion", "Staging base release pak files requires a base release version to be specified").ToString();
 		case ELauncherProfileValidationErrors::GeneratingChunksRequiresCookByTheBook:
-			return LOCTEXT("GeneratingChunksRequiresCookByTheBook", "Generating Chunks requires cook by the book mode.").ToString();
+			return LOCTEXT("GeneratingChunksRequiresCookByTheBook", "Generating Chunks requires cooking or zen snapshot import.").ToString();
 		case ELauncherProfileValidationErrors::GeneratingChunksRequiresUnrealPak:
 			return LOCTEXT("GeneratingChunksRequiresUnrealPak", "UnrealPak must be selected to Generate Chunks.").ToString();
 		case ELauncherProfileValidationErrors::GeneratingHttpChunkDataRequiresGeneratingChunks:
diff --git a/Engine/Source/Developer/LauncherServices/Private/Profiles/LauncherProfile.h b/Engine/Source/Developer/LauncherServices/Private/Profiles/LauncherProfile.h
--- a/Engine/Source/Developer/LauncherServices/Private/Profiles/LauncherProfile.h
+++ b/Engine/Source/Developer/LauncherServices/Private/Profiles/LauncherProfile.h
@@ -733,6 +733,61 @@ public:
 		return ValidationErrors.Contains(Error);
 	}
 
+	virtual void AddCustomError( const FString& UniqueId, const FText& Text ) override
+	{
+		CustomErrors.Add(UniqueId, Text);
+	}
+
+	virtual bool HasCustomError( const FString& UniqueId ) const override
+	{
+		return CustomErrors.Contains(UniqueId);
+	}
+
+	virtual FText GetCustomErrorText( const FString& UniqueId ) const override
+	{
+		const FText* TextPtr = CustomErrors.Find(UniqueId);
+		return TextPtr ? (*TextPtr) : FText::GetEmpty();
+	}
+
+	virtual TArray<FString> GetAllCustomErrors() const override
+	{
+		TArray<FString> Keys;
+		CustomErrors.GetKeys(Keys);
+		return MoveTemp(Keys);
+	}
+
+	virtual void AddCustomWarning( const FString& UniqueId, const FText& Text ) override
+	{
+		CustomWarnings.Add(UniqueId, Text);
+	}
+
+	virtual bool HasCustomWarning( const FString& UniqueId ) const override
+	{
+		return CustomWarnings.Contains(UniqueId);
+	}
+
+	virtual FText GetCustomWarningText( const FString& UniqueId ) const override
+	{
+		const FText* TextPtr = CustomWarnings.Find(UniqueId);
+		return TextPtr ? (*TextPtr) : FText::GetEmpty();
+	}
+
+	virtual TArray<FString> GetAllCustomWarnings() const override
+	{
+		TArray<FString> Keys;
+		CustomWarnings.GetKeys(Keys);
+		return MoveTemp(Keys);
+	}
+
+	virtual void RefreshCustomWarningsAndErrors() override
+	{
+		CustomWarnings.Reset();
+		CustomErrors.Reset();
+		CustomValidationDelegate.Broadcast();
+	}
+
+
+
 	virtual FString GetInvalidPlatform() const override
 	{
 		return InvalidPlatform;
@@ -915,7 +970,7 @@ public:
 
 	virtual bool IsValidForLaunch( ) override
 	{
-		return (ValidationErrors.Num() == 0);
+		return (ValidationErrors.Num() == 0 && CustomErrors.Num() == 0);
 	}
 
 	virtual void RemoveCookedCulture( const FString& CultureName ) override
@@ -2775,6 +2830,11 @@ public:
 		return BuildTargetOptionsChangedDelegate;
 	}
 
+	virtual FOnProfileCustomValidation& OnCustomValidation() override
+	{
+		return CustomValidationDelegate;
+	}
+
 	virtual void SetEditorExe( const FString& InEditorExe ) override
 	{
 		EditorExe = InEditorExe;
@@ -3176,6 +3236,8 @@ protected:
 		ValidateBuildTarget();
 		ValidatePlatformSDKs();
 		ValidateDeviceStatus();
+
+		RefreshCustomWarningsAndErrors();
 	}
 	
 	void ValidatePlatformSDKs(void)
@@ -3586,6 +3648,10 @@ private:
 	TArray<ELauncherProfileValidationErrors::Type> ValidationErrors;
 	FString InvalidPlatform;
 
+	// Holds the collection of custom errors & warnings.
+	TMap<FString,FText> CustomErrors;
+	TMap<FString,FText> CustomWarnings;
+
     // Holds the time out time for the cook on the fly server
     uint32 Timeout;
 
@@ -3651,6 +3717,9 @@ private:
 	// Holds a delegate to be invoked when the project build target options have changed
 	FOnProfileBuildTargetOptionsChanged BuildTargetOptionsChangedDelegate;
 
+	// Holds a delegate to be invoked when the profile is validated to gather custom warnings and errors
+	FOnProfileCustomValidation CustomValidationDelegate;
+
 	// Cached build target options (not serialized)
 	TArray<FString> ExplictBuildTargetNames;
 };
diff --git a/Engine/Source/Developer/LauncherServices/Public/ILauncherProfile.h b/Engine/Source/Developer/LauncherServices/Public/ILauncherProfile.h
--- a/Engine/Source/Developer/LauncherServices/Public/ILauncherProfile.h
+++ b/Engine/Source/Developer/LauncherServices/Public/ILauncherProfile.h
@@ -399,6 +399,10 @@ DECLARE_MULTICAST_DELEGATE(FOnProfileProjectChanged);
 /** Delegate type for a change in build target options */
 DECLARE_MULTICAST_DELEGATE(FOnProfileBuildTargetOptionsChanged);
 
+/** Delegate type for gathering custom profile validation */
+DECLARE_MULTICAST_DELEGATE(FOnProfileCustomValidation);
+
+
 /** Delegate type for detecting if cook is finished
  *	Used when cooking from the editor.  Specific cook task will wait for the cook to be finished by the editor
  */
@@ -467,6 +471,78 @@ public:
 	 */
 	virtual bool HasValidationError( ELauncherProfileValidationErrors::Type Error ) const = 0;
 
+
+	/**
+	 * Add a custom error. This will prevent launching
+	 * 
+	 * @param UniqueId Unique identifier for this error
+	 * @param Text Localized string describing this error
+	 */
+	virtual void AddCustomError( const FString& UniqueId, const FText& Text ) = 0;
+
+
+	/**
+	 * Checks whether this profile has the given custom error
+	 *
+	 * @param UniqueId Unique identifier for this error
+	 */
+	virtual bool HasCustomError( const FString& UniqueId ) const = 0;
+
+	/**
+	 * Gets the localized string representation of the given error
+	 * 
+	 * @param UniqueId
+	 * @returns Localized string describing this error, or empty text if we don't have it
+	 */
+	virtual FText GetCustomErrorText( const FString& UniqueId ) const = 0;
+
+	/**
+	 * Gets all custom errors in this profile
+	 * 
+	 * @returns UniqueIds for all custom errors in this profile
+	 */
+	virtual TArray<FString> GetAllCustomErrors() const = 0;
+	
+
+
+	/**
+	 * Add a custom warning. This does not prevent launching.
+	 * 
+	 * @param UniqueId Unique identifier for this warning
+	 * @param Text Localized string describing this warning
+	 */
+	virtual void AddCustomWarning( const FString& UniqueId, const FText& Text ) = 0;
+
+	/**
+	 * Checks whether this profile has the given custom warning
+	 *
+	 * @param UniqueId Unique identifier for this warning
+	 */
+	virtual bool HasCustomWarning( const FString& UniqueId ) const = 0;
+
+	/**
+	 * Gets the localized string representation of the given warning
+	 * 
+	 * @param UniqueId
+	 * @returns Localized string describing this warning, or empty text if we don't have it
+	 */
+	virtual FText GetCustomWarningText( const FString& UniqueId ) const = 0;
+
+	/**
+	 * Gets all custom warnings in this profile
+	 * 
+	 * @returns UniqueIds for all custom warnings in this profile
+	 */
+	virtual TArray<FString> GetAllCustomWarnings() const = 0;
+
+
+	/**
+	 * Refreshes the custom warnings and errors. Resets the lists and triggers the OnValidation callback.
+	 */
+	virtual void RefreshCustomWarningsAndErrors() = 0;
+
+
+
 	/**
 	 * Gets the invalid platform, this is only valid when there is a platform centric validation error.
 	 *
@@ -1396,6 +1472,13 @@ public:
 	 */
 	virtual FOnProfileBuildTargetOptionsChanged& OnBuildTargetOptionsChanged() = 0;
 
+	/**
+	 * Access delegate used when validation changes to gather custom warnings and errors
+	 * 
+	 * @return The delegate.
+	 */
+	virtual FOnProfileCustomValidation& OnCustomValidation() = 0;
+
 	/**
 	 * Sets whether to use I/O store for optimized loading.
 	 * @param bUseIoStore Whether to use I/O store
