Add Int and Float support to Project Launcher 2 property & config fields

#rnx

diff --git a/Engine/Plugins/Developer/ProjectLauncher/Source/ProjectLauncher/Private/Extension/LaunchExtension.cpp b/Engine/Plugins/Developer/ProjectLauncher/Source/ProjectLauncher/Private/Extension/LaunchExtension.cpp
--- a/Engine/Plugins/Developer/ProjectLauncher/Source/ProjectLauncher/Private/Extension/LaunchExtension.cpp
+++ b/Engine/Plugins/Developer/ProjectLauncher/Source/ProjectLauncher/Private/Extension/LaunchExtension.cpp
@@ -483,6 +483,40 @@ namespace ProjectLauncher
 		}
 	}
 
+	int32 FLaunchExtensionInstance::GetConfigInteger( EConfig Config, const TCHAR* Name, int32 DefaultValue ) const
+	{
+		FString KeyName = GetConfigKeyName(Config, Name);
+
+		if (Config == EConfig::PerProfile)
+		{
+			const FString* ValuePtr = Profile->GetCustomStringProperties().Find( KeyName );
+			return ValuePtr ? FCString::Atoi(ValuePtr->GetCharArray().GetData()) : DefaultValue;
+
+		}
+		else
+		{
+			int32 Value;
+			return GConfig->GetInt( Model->GetConfigSection() , *KeyName, Value, Model->GetConfigIni() ) ? Value : DefaultValue;
+		}
+	}
+
+	float FLaunchExtensionInstance::GetConfigFloat( EConfig Config, const TCHAR* Name, float DefaultValue ) const
+	{
+		FString KeyName = GetConfigKeyName(Config, Name);
+
+		if (Config == EConfig::PerProfile)
+		{
+			const FString* ValuePtr = Profile->GetCustomStringProperties().Find( KeyName );
+			return ValuePtr ? FCString::Atof(ValuePtr->GetCharArray().GetData()) : DefaultValue;
+
+		}
+		else
+		{
+			float Value;
+			return GConfig->GetFloat( Model->GetConfigSection() , *KeyName, Value, Model->GetConfigIni() ) ? Value : DefaultValue;
+		}
+	}
+
 
 	void FLaunchExtensionInstance::SetConfigString( EConfig Config, const TCHAR* Name, const FString& Value ) const
 	{
@@ -502,9 +536,10 @@ namespace ProjectLauncher
 		}
 		else
 		{
-			GConfig->SetString( Model->GetConfigSection() , *KeyName, *Value, Model->GetConfigIni() );
+			GConfig->SetString( Model->GetConfigSection(), *KeyName, *Value, Model->GetConfigIni() );
 		}
 	}
+
 	void FLaunchExtensionInstance::SetConfigBool( EConfig Config, const TCHAR* Name, const bool& Value ) const
 	{
 		FString KeyName = GetConfigKeyName(Config, Name);
@@ -527,6 +562,50 @@ namespace ProjectLauncher
 		}
 	}
 
+	void FLaunchExtensionInstance::SetConfigInteger( EConfig Config, const TCHAR* Name, const int32& Value ) const
+	{
+		FString KeyName = GetConfigKeyName(Config, Name);
+
+		if (Config == EConfig::PerProfile)
+		{
+			if (Value == 0)
+			{
+				Profile->GetCustomStringProperties().Remove(KeyName);
+			}
+			else
+			{
+				Profile->GetCustomStringProperties().Add(KeyName, ::LexToString(Value) );
+			}
+			BroadcastPropertyChanged();
+		}
+		else
+		{
+			GConfig->SetInt( Model->GetConfigSection() , *KeyName, Value, Model->GetConfigIni() );
+		}
+	}
+
+	void FLaunchExtensionInstance::SetConfigFloat( EConfig Config, const TCHAR* Name, const float& Value ) const
+	{
+		FString KeyName = GetConfigKeyName(Config, Name);
+
+		if (Config == EConfig::PerProfile)
+		{
+			if (Value == 0.0f)
+			{
+				Profile->GetCustomStringProperties().Remove(KeyName);
+			}
+			else
+			{
+				Profile->GetCustomStringProperties().Add(KeyName, ::LexToString(Value) );
+			}
+			BroadcastPropertyChanged();
+		}
+		else
+		{
+			GConfig->SetFloat( Model->GetConfigSection() , *KeyName, Value, Model->GetConfigIni() );
+		}
+	}
+
 	FString FLaunchExtensionInstance::GetConfigKeyName( EConfig Config, const TCHAR* Name ) const
 	{
 		switch (Config)
diff --git a/Engine/Plugins/Developer/ProjectLauncher/Source/ProjectLauncher/Private/ProfileTree/LaunchProfileTreeData.cpp b/Engine/Plugins/Developer/ProjectLauncher/Source/ProjectLauncher/Private/ProfileTree/LaunchProfileTreeData.cpp
--- a/Engine/Plugins/Developer/ProjectLauncher/Source/ProjectLauncher/Private/ProfileTree/LaunchProfileTreeData.cpp
+++ b/Engine/Plugins/Developer/ProjectLauncher/Source/ProjectLauncher/Private/ProfileTree/LaunchProfileTreeData.cpp
@@ -11,6 +11,7 @@
 #include "Widgets/Input/SEditableTextBox.h"
 #include "Widgets/Input/SMultiLineEditableTextBox.h"
 #include "Widgets/Input/SComboButton.h"
+#include "Widgets/Input/SSpinBox.h"
 #include "Widgets/Images/SImage.h"
 #include "Widgets/SBoxPanel.h"
 #include "Layout/Visibility.h"
@@ -184,6 +185,11 @@ namespace ProjectLauncher
 	{
 	}
 
+	void FLaunchProfileTreeNode::OnPropertyChanged()
+	{
+		TreeData->TreeBuilder->OnPropertyChanged();
+	}
+
 	FLaunchProfileTreeNode& FLaunchProfileTreeNode::AddWidget( FText InName, FCallbacks&& InWidgetCallbacks, TSharedRef<SWidget> InValueWidget )
 	{
 		FLaunchProfileTreeNodePtr TreeNode = MakeShared<FLaunchProfileTreeNode>(TreeData);
@@ -230,7 +236,7 @@ namespace ProjectLauncher
 		auto SetCheckState = [this, BooleanCallbacks]( ECheckBoxState CheckState )
 		{
 			BooleanCallbacks.SetValue( (CheckState == ECheckBoxState::Checked) );
-			TreeData->TreeBuilder->OnPropertyChanged();
+			OnPropertyChanged();
 		};
 
 		auto GetCheckState = [BooleanCallbacks]()
@@ -274,7 +280,7 @@ namespace ProjectLauncher
 		auto SetString = [this, StringCallbacks]( const FText& InText, ETextCommit::Type InTextCommit )
 		{
 			StringCallbacks.SetValue( InText.ToString() );
-			TreeData->TreeBuilder->OnPropertyChanged();
+			OnPropertyChanged();
 		};
 
 		auto GetString = [StringCallbacks]()
@@ -322,7 +328,7 @@ namespace ProjectLauncher
 		auto SetString = [this, StringCallbacks]( const FText& InText, ETextCommit::Type InTextCommit )
 		{
 			StringCallbacks.SetValue( InText.ToString() );
-			TreeData->TreeBuilder->OnPropertyChanged();
+			OnPropertyChanged();
 		};
 
 		auto GetString = [StringCallbacks]()
@@ -352,7 +358,7 @@ namespace ProjectLauncher
 				}
 
 				StringCallbacks.SetValue( OutDirectory );
-				TreeData->TreeBuilder->OnPropertyChanged();
+				OnPropertyChanged();
 			}
 
 			return FReply::Handled();
@@ -428,7 +434,7 @@ namespace ProjectLauncher
 		auto SetString = [this, StringCallbacks]( const FText& InText, ETextCommit::Type InTextCommit )
 		{
 			StringCallbacks.SetValue( InText.ToString() );
-			TreeData->TreeBuilder->OnPropertyChanged();
+			OnPropertyChanged();
 		};
 
 		auto GetString = [StringCallbacks]()
@@ -507,6 +513,27 @@ namespace ProjectLauncher
 		);
 	}
 
+
+
+
+
+
+
+
+
+
+	FLaunchProfileTreeNode& FLaunchProfileTreeNode::AddInteger( FText InName, TTypeCallbacks<int32>&& TypeCallbacks, FInt32Interval Range, bool bShowSpinner )
+	{
+		return AddValue<int32>(InName, MoveTemp(TypeCallbacks), Range, bShowSpinner);
+	}
+
+
+	FLaunchProfileTreeNode& FLaunchProfileTreeNode::AddFloat( FText InName, TTypeCallbacks<float>&& TypeCallbacks, FFloatInterval Range, bool bShowSpinner )
+	{
+		return AddValue<float>(InName, MoveTemp(TypeCallbacks), Range, bShowSpinner);
+	}
+
+
 }
 
 #undef LOCTEXT_NAMESPACE
diff --git a/Engine/Plugins/Developer/ProjectLauncher/Source/ProjectLauncher/Public/Extension/LaunchExtension.h b/Engine/Plugins/Developer/ProjectLauncher/Source/ProjectLauncher/Public/Extension/LaunchExtension.h
--- a/Engine/Plugins/Developer/ProjectLauncher/Source/ProjectLauncher/Public/Extension/LaunchExtension.h
+++ b/Engine/Plugins/Developer/ProjectLauncher/Source/ProjectLauncher/Public/Extension/LaunchExtension.h
@@ -235,6 +235,26 @@ namespace ProjectLauncher
 		 */
 		UE_API bool GetConfigBool( EConfig Config, const TCHAR* Name, bool DefaultValue = false ) const;
 
+		/**
+		 * Read a configuration int value
+		 * 
+		 * @param Config where the value is stored
+		 * @param Name name of the value to read
+		 * @param DefaultValue value to return if the value isn't found
+		 * @returns the value of the bool, or DefaultValue if it isn't found
+		 */
+		UE_API int32 GetConfigInteger( EConfig Config, const TCHAR* Name, int32 DefaultValue = 0 ) const;
+
+		/**
+		 * Read a configuration float value
+		 * 
+		 * @param Config where the value is stored
+		 * @param Name name of the value to read
+		 * @param DefaultValue value to return if the value isn't found
+		 * @returns the value of the bool, or DefaultValue if it isn't found
+		 */
+		UE_API float GetConfigFloat( EConfig Config, const TCHAR* Name, float DefaultValue = 0.0f ) const;
+
 		/**
 		 * Write a configuration string value
 		 * 
@@ -253,6 +273,26 @@ namespace ProjectLauncher
 		 */
 		UE_API void SetConfigBool( EConfig Config, const TCHAR* Name, const bool& Value ) const;
 
+		/**
+		 * Write a configuration int value
+		 * 
+		 * @param Config where the value is stored
+		 * @param Name name of the value to write
+		 * @param Value value to write
+		 */
+		UE_API void SetConfigInteger( EConfig Config, const TCHAR* Name, const int32& Value ) const;
+
+		/**
+		 * Write a configuration float value
+		 * 
+		 * @param Config where the value is stored
+		 * @param Name name of the value to write
+		 * @param Value value to write
+		 */
+		UE_API void SetConfigFloat( EConfig Config, const TCHAR* Name, const float& Value ) const;
+
+
+
 		/**
 		 * Get the final key name to use for reading & writing a configuration value
 		 * 
diff --git a/Engine/Plugins/Developer/ProjectLauncher/Source/ProjectLauncher/Public/ProfileTree/LaunchProfileTreeData.h b/Engine/Plugins/Developer/ProjectLauncher/Source/ProjectLauncher/Public/ProfileTree/LaunchProfileTreeData.h
--- a/Engine/Plugins/Developer/ProjectLauncher/Source/ProjectLauncher/Public/ProfileTree/LaunchProfileTreeData.h
+++ b/Engine/Plugins/Developer/ProjectLauncher/Source/ProjectLauncher/Public/ProfileTree/LaunchProfileTreeData.h
@@ -5,6 +5,8 @@
 #include "Internationalization/Text.h"
 #include "Templates/Function.h"
 #include "Templates/SharedPointer.h"
+#include "Math/Interval.h"
+#include "Widgets/Input/SSpinBox.h"
 #include "ILauncherProfile.h"
 #include <initializer_list>
 
@@ -56,6 +58,10 @@ namespace ProjectLauncher
 		typedef TFunction<void()> FFunction;
 		typedef TFunction<bool()> FGetBool;
 		typedef TFunction<void(bool)> FSetBool;
+		typedef TFunction<int32()> FGetInt;
+		typedef TFunction<void(int32)> FSetInt;
+		typedef TFunction<float()> FGetFloat;
+		typedef TFunction<void(float)> FSetFloat;
 		typedef TFunction<FString(void)> FGetString;
 		typedef TFunction<void(FString)> FSetString;
 
@@ -83,6 +89,7 @@ namespace ProjectLauncher
 		};
 		UE_API FLaunchProfileTreeNode& AddBoolean( FText InName, FBooleanCallbacks&& BooleanCallbacks );
 
+
 		struct FStringCallbacks
 		{
 			FGetString GetValue = nullptr;
@@ -97,6 +104,24 @@ namespace ProjectLauncher
 		UE_API FLaunchProfileTreeNode& AddCommandLineString( FText InName, FStringCallbacks&& StringCallbacks );
 
 
+		template<typename T>
+		struct TTypeCallbacks
+		{
+			TFunction<T()> GetValue = nullptr;
+			TFunction<void(T)> SetValue = nullptr;
+			TFunction<T()> GetDefaultValue = nullptr;
+			FGetBool IsVisible = nullptr;
+			FGetBool IsEnabled = nullptr;
+			FValidation Validation;
+		};
+		template<typename T>
+		FLaunchProfileTreeNode& AddValue( FText InName, TTypeCallbacks<T>&& TypeCallbacks, TInterval<T> Range = TInterval<T>(), bool bShowSpinner = false );
+
+
+		UE_API FLaunchProfileTreeNode& AddInteger( FText InName, TTypeCallbacks<int32>&& TypeCallbacks, FInt32Interval Range = FInt32Interval(), bool bShowSpinner = false );
+		UE_API FLaunchProfileTreeNode& AddFloat( FText InName, TTypeCallbacks<float>&& TypeCallbacks, FFloatInterval Range = FFloatInterval(), bool bShowSpinner = false );
+
+
 
 		// member variables
 		FText Name;
@@ -109,6 +134,8 @@ namespace ProjectLauncher
 		const FLaunchProfileTreeData* GetTreeData() const { return TreeData; }
 
 	protected:
+		UE_API void OnPropertyChanged();
+
 		FLaunchProfileTreeData* TreeData;
 	};
 
@@ -139,6 +166,81 @@ namespace ProjectLauncher
 	};
 
 
+
+
+	template<typename T>
+	FLaunchProfileTreeNode& FLaunchProfileTreeNode::AddValue( FText InName, TTypeCallbacks<T>&& TypeCallbacks, TInterval<T> Range, bool bShowSpinner )
+	{
+		check(TypeCallbacks.GetValue);
+		check(TypeCallbacks.SetValue);
+
+		FCallbacks WidgetCallbacks
+		{
+			.IsVisible = TypeCallbacks.IsVisible,
+			.IsEnabled = TypeCallbacks.IsEnabled,
+			.Validation = TypeCallbacks.Validation,
+		};
+		if (TypeCallbacks.GetDefaultValue != nullptr)
+		{
+			WidgetCallbacks.IsDefault = [TypeCallbacks]()
+			{
+				return TypeCallbacks.GetValue() == TypeCallbacks.GetDefaultValue();
+			};
+
+			WidgetCallbacks.SetToDefault = [TypeCallbacks]()
+			{
+				TypeCallbacks.SetValue(TypeCallbacks.GetDefaultValue());
+				// nb. OnPropertyChanged is called via SCustomLaunchWidgetTableRow's set to default
+			};
+		};
+
+		auto SetImmediateValue = [this, Range, TypeCallbacks]( T Value )
+		{
+			if (Range.IsValid())
+			{
+				TypeCallbacks.SetValue( FMath::Clamp( Value, Range.Min, Range.Max) );
+			}
+			else
+			{
+				TypeCallbacks.SetValue( Value );
+			}
+		};
+
+		auto SetValue = [this, SetImmediateValue]( T Value, ETextCommit::Type )
+		{
+			SetImmediateValue(Value);
+			OnPropertyChanged();
+		};
+
+		auto GetValue = [TypeCallbacks]()
+		{
+			return TypeCallbacks.GetValue();
+		};
+
+		if (Range.IsValid())
+		{
+			return AddWidget( InName,
+				MoveTemp(WidgetCallbacks),
+				SNew(SSpinBox<T>)
+				.Value_Lambda(GetValue)
+				.OnValueChanged_Lambda(SetImmediateValue)
+				.OnValueCommitted_Lambda(SetValue)
+				.MinValue(Range.Min)
+				.MaxValue(Range.Max)
+				.EnableSlider(bShowSpinner)
+			);
+		}
+		else
+		{
+			return AddWidget( InName,
+				MoveTemp(WidgetCallbacks),
+				SNew(SSpinBox<T>)
+				.Value_Lambda(GetValue)
+				.OnValueCommitted_Lambda(SetValue)
+				.EnableSlider(false)
+			);
+		}
+	}
 }
 
 #undef UE_API
