AutomatedPerfTest Launcher Extension

* Add Project Launcher Extension to launch Automated perf test
* Allow specifiying test args via Launcher UI

#rnx

diff --git a/Engine/Plugins/Developer/ProjectLauncher/ProjectLauncher.uplugin b/Engine/Plugins/Developer/ProjectLauncher/ProjectLauncher.uplugin
--- a/Engine/Plugins/Developer/ProjectLauncher/ProjectLauncher.uplugin
+++ b/Engine/Plugins/Developer/ProjectLauncher/ProjectLauncher.uplugin
@@ -44,5 +44,11 @@
 			],
 			"ProgramAllowList": [ "UnrealFrontend" ]
 		}
+	],
+	"Plugins": [
+		{
+			"Name": "AutomatedPerfTesting",
+			"Enabled": true
+		}
 	]
 }
\ No newline at end of file
diff --git a/Engine/Plugins/Developer/ProjectLauncher/Source/CommonLaunchExtensions/CommonLaunchExtensions.Build.cs b/Engine/Plugins/Developer/ProjectLauncher/Source/CommonLaunchExtensions/CommonLaunchExtensions.Build.cs
--- a/Engine/Plugins/Developer/ProjectLauncher/Source/CommonLaunchExtensions/CommonLaunchExtensions.Build.cs
+++ b/Engine/Plugins/Developer/ProjectLauncher/Source/CommonLaunchExtensions/CommonLaunchExtensions.Build.cs
@@ -9,11 +9,25 @@ public class CommonLaunchExtensions : ModuleRules
 		PrivateDependencyModuleNames.AddRange(
 		[
 			"Core",
+			"CoreUObject",
+			"Json",
 			"Slate",
 			"SlateCore",
 			"Sockets",
 			"TraceLog",
 			"ProjectLauncher",
+			"DeveloperSettings"
 		]);
+
+
+		if (Target.bBuildEditor)
+		{
+			PrivateDependencyModuleNames.AddRange(
+			[
+				"AutomatedPerfTesting",
+				"UnrealEd",
+				"Engine"
+			]);
+		}
 	}
 }
diff --git a/Engine/Plugins/Developer/ProjectLauncher/Source/CommonLaunchExtensions/Private/AutomatedPerfTest/AutomatedPerfTestLaunchExtension.cpp b/Engine/Plugins/Developer/ProjectLauncher/Source/CommonLaunchExtensions/Private/AutomatedPerfTest/AutomatedPerfTestLaunchExtension.cpp
new file mode 100644
--- /dev/null
+++ b/Engine/Plugins/Developer/ProjectLauncher/Source/CommonLaunchExtensions/Private/AutomatedPerfTest/AutomatedPerfTestLaunchExtension.cpp
@@ -0,0 +1,526 @@
+// Copyright Epic Games, Inc. All Rights Reserved.
+
+#include "AutomatedPerfTestLaunchExtension.h"
+#include "CoreMinimal.h"
+
+#if WITH_EDITOR
+
+#include "AutomatedProfileGoTest.h"
+#include "AutomatedMaterialPerfTest.h"
+#include "AutomatedReplayPerfTest.h"
+#include "AutomatedSequencePerfTest.h"
+#include "StaticCameraTests/AutomatedStaticCameraPerfTestBase.h"
+#include "StaticCameraTests/AutomatedPlacedStaticCameraPerfTest.h"
+
+#include "Engine/World.h" 
+#include "Editor/EditorEngine.h"
+#include "IPropertyTypeCustomization.h"
+#include "PropertyHandle.h"
+#include "Kismet/GameplayStatics.h"
+
+#endif // WITH_EDITOR
+
+#include "HAL/PlatformMisc.h"
+#include "HAL/PlatformProcess.h"
+#include "Templates/SharedPointer.h"
+
+#include "SlateOptMacros.h"
+#include "Widgets/SCompoundWidget.h"
+#include "Styling/SlateTypes.h"
+#include "Styling/CoreStyle.h"
+#include "Styling/SlateWidgetStyleAsset.h"
+#include "Fonts/SlateFontInfo.h"
+#include "Model/ProjectLauncherModel.h"
+
+#include "Misc/Attribute.h"
+#include "Misc/MessageDialog.h"
+#include "Misc/PackageName.h"
+#include "Misc/Paths.h"
+
+#include "Styling/AppStyle.h"
+#include "Widgets/DeclarativeSyntaxSupport.h"
+#include "Widgets/Input/SComboButton.h"
+#include "Widgets/Input/SButton.h"
+#include "Widgets/Text/STextBlock.h"
+#include "Widgets/Shared/SCustomLaunchCombo.h"
+
+#include "Containers/UnrealString.h"
+#include "Logging/LogMacros.h"
+
+#define LOCTEXT_NAMESPACE "FAutomatedPerfTestLaunchExtension"
+
+namespace AutomatedPerfTestConfig
+{
+	constexpr const TCHAR* AutomatedPerfTestInternalName = TEXT("AutomatedPerfTestExtension.PerfTest");
+	constexpr const TCHAR* EnableCSVProfilerConfig = TEXT("EnableCSVProfiler");
+	constexpr const TCHAR* WindowedConfig = TEXT("Windowed");
+	constexpr const TCHAR* IterationsConfig = TEXT("Iterations");
+
+	constexpr const TCHAR* ProfileGoMapConfig = TEXT("ProfileGoMapName");
+	constexpr const TCHAR* ReplayConfig = TEXT("ReplayTestName");
+	constexpr const TCHAR* SequenceTestNameConfig = TEXT("SequenceTestName");
+	constexpr const TCHAR* StaticCameraMapConfig = TEXT("StaticCameraMapName");
+	constexpr const TCHAR* ProfileGoConfigFileName = TEXT("ProfileGoConfigFile");
+
+	constexpr EAutomatedPerfTestType DefaultTestType = EAutomatedPerfTestType::Sequence;
+	constexpr uint32 DefaultIterationCount = 1;
+
+	static constexpr bool IsEditorBuild()
+	{
+#if WITH_EDITOR
+		constexpr bool bIsEditor = true;
+#else
+		constexpr bool bIsEditor = false;
+#endif
+		return bIsEditor;
+	}
+}
+
+using namespace AutomatedPerfTestConfig;
+
+void FAutomatedPerfTestLaunchExtensionInstance::MakeCustomExtensionSubmenu(FMenuBuilder& MenuBuilder)
+{
+	const auto ToggleTestEnabled = [this]()
+	{
+		ILauncherProfileAutomatedTestPtr ExistingAutomatedTest = GetTest();
+		if (ExistingAutomatedTest.IsValid())
+		{
+			OnTestRemoved(ExistingAutomatedTest.ToSharedRef());
+			GetProfile()->RemoveAutomatedTest(*GetTestInternalName());
+		}
+		else
+		{
+			ILauncherProfileAutomatedTestRef AutomatedTest = GetProfile()->FindOrAddAutomatedTest(*GetTestInternalName());
+			OnTestAdded(AutomatedTest);
+		}
+
+		BroadcastPropertyChanged();
+	};
+
+	MenuBuilder.AddMenuEntry(
+		LOCTEXT("EnableAutomatedPerfTestLabel", "Automated Perf Test"),
+		FText::GetEmpty(),
+		FSlateIcon(),
+		FUIAction(
+			FExecuteAction::CreateLambda(ToggleTestEnabled),
+			FCanExecuteAction(),
+			FIsActionChecked::CreateRaw(this, &FAutomatedTestLaunchExtensionInstance::IsTestActive)
+		),
+		NAME_None,
+		EUserInterfaceActionType::Check);
+}
+
+void FAutomatedPerfTestLaunchExtensionInstance::OnPropertyChanged()
+{
+	ILauncherProfileAutomatedTestPtr Test = GetTest();
+	if(Test.IsValid())
+	{
+		// This ensures the test type is always 
+		// updated to the current selected one. 
+		OnTestAdded(Test.ToSharedRef());
+	}
+
+	Super::OnPropertyChanged();
+}
+
+static constexpr const TCHAR* GetTestNodeName(EAutomatedPerfTestType TestType)
+{
+	// This is the string which represents test node defined in the
+	// Gauntlet controller. The test launch will fail if this string
+	// does not match the name of the controller class in Gauntlet.
+	switch (TestType)
+	{
+	case EAutomatedPerfTestType::Sequence:		return TEXT("AutomatedPerfTest.SequenceTest");
+	case EAutomatedPerfTestType::Replay:		return TEXT("AutomatedPerfTest.ReplayTest");
+	case EAutomatedPerfTestType::StaticCamera:	return TEXT("AutomatedPerfTest.StaticCameraTest");
+	case EAutomatedPerfTestType::Material:		return TEXT("AutomatedPerfTest.MaterialTest");
+	case EAutomatedPerfTestType::ProfileGo:		return TEXT("AutomatedPerfTest.ProfileGoTest");
+	default:									return TEXT("AutomatedPerfTest.DefaultTest");
+	}
+}
+
+void FAutomatedPerfTestLaunchExtensionInstance::OnTestAdded(ILauncherProfileAutomatedTestRef AutomatedTest)
+{
+    constexpr int32 DefaultPriority = 1000;
+	const TCHAR* TestNode = GetTestNodeName(GetTestType());
+
+	AutomatedTest->SetTests(TestNode);
+	AutomatedTest->SetPriority(DefaultPriority);
+}
+
+const FString FAutomatedPerfTestLaunchExtensionInstance::GetTestInternalName() const
+{
+	return AutomatedPerfTestInternalName;
+}
+
+void FAutomatedPerfTestLaunchExtensionInstance::CustomizeTree(ProjectLauncher::FLaunchProfileTreeData& ProfileTreeData)
+{
+	using namespace ProjectLauncher;
+	const auto IsVisible = [this]()
+	{
+		return IsTestActive() || GetConfigBool(EConfig::PerProfile, TEXT("Enabled") );
+	};
+
+	const auto GetDisplayName = [](EAutomatedPerfTestType Type)
+	{
+		switch (Type)
+		{
+			case EAutomatedPerfTestType::Sequence:		return LOCTEXT("SequenceLabel", "Sequence");
+			case EAutomatedPerfTestType::Replay:		return LOCTEXT("ReplayLabel", "Replay");
+			case EAutomatedPerfTestType::ProfileGo:		return LOCTEXT("ProfileGoLabel", "ProfileGo");
+			case EAutomatedPerfTestType::StaticCamera:	return LOCTEXT("StaticCameraLabel", "StaticCamera");
+			case EAutomatedPerfTestType::Material:		return LOCTEXT("MaterialLabel", "Material");
+		}
+		return FText::GetEmpty();
+	};
+
+	const auto GetToolTip = [](EAutomatedPerfTestType Type)
+	{
+		switch (Type)
+		{
+		case EAutomatedPerfTestType::Sequence:		return LOCTEXT("SequenceToolTipLabel", "Automated Sequence Perf Test");
+		case EAutomatedPerfTestType::Replay:		return LOCTEXT("ReplayToolTipLabel", "Automated Replay Perf Test");
+		case EAutomatedPerfTestType::ProfileGo:		return LOCTEXT("ProfileGoToolTipLabel", "Automated ProfileGo Perf Test");
+		case EAutomatedPerfTestType::StaticCamera:	return LOCTEXT("StaticCameraToolTipLabel", "Automated StaticCamera Perf Test");
+		case EAutomatedPerfTestType::Material:		return LOCTEXT("MaterialToolTipLabel", "Automated Material Perf Test");
+		}
+		return FText::GetEmpty();
+	};
+
+	// TODO: Add support to specify sub test type i.e. LLM/Insights/GPU Perf etc. 
+	// Adding it as enum like the main test type above may not be ideal as sub-tests
+	// are expected to change more often. It may be prudent to consider that solution
+	// for the test types above as well to ensure test types can be added on and is 
+	// extensible in the future.
+
+	constexpr EConfig ConfigTypeCommon = EConfig::User_PerProfile;
+	FLaunchProfileTreeNode& TreeNode = ProfileTreeData
+		.AddHeading(GetExtension()->GetDisplayName())
+		.AddBoolean(LOCTEXT("APTEnableCSVLabel", "Enable CSV Profiler"), 
+		{
+			.GetValue = [this]() { return GetConfigBool(ConfigTypeCommon, EnableCSVProfilerConfig, true); },
+			.SetValue = [this](bool bEnable) { return SetConfigBool(ConfigTypeCommon, EnableCSVProfilerConfig, bEnable); },
+			.IsVisible = IsVisible,
+		})
+		.AddBoolean(LOCTEXT("WindowedLabel","Windowed"),
+		{
+			.GetValue = [this]()          { return GetConfigBool(ConfigTypeCommon, WindowedConfig);},
+			.SetValue = [this](bool bVal) { SetConfigBool(ConfigTypeCommon, WindowedConfig, bVal ); },
+			.IsVisible = IsVisible,
+		})
+		.AddInteger(LOCTEXT("IterationsLabel", "Iterations"), 
+		{
+			.GetValue = [this]() { return GetConfigInteger(ConfigTypeCommon, IterationsConfig, DefaultIterationCount); },
+			.SetValue = [this](int32 Value) { SetConfigInteger(ConfigTypeCommon, IterationsConfig, Value); },
+			.IsVisible = IsVisible,
+		})
+		.AddWidget(LOCTEXT("TestTypeLabel", "Test Type"),
+		{
+			.IsDefault = [this]() { return GetTestType() == DefaultTestType; },
+			.SetToDefault = [this]() { SetTestType(DefaultTestType); },
+			.IsVisible = IsVisible
+		},
+		SNew(SCustomLaunchCombo<EAutomatedPerfTestType>)
+		.OnSelectionChanged(this, &FAutomatedPerfTestLaunchExtensionInstance::OnTestTypeSelectionChanged)
+		.SelectedItem(this, &FAutomatedPerfTestLaunchExtensionInstance::GetTestType)
+		.GetDisplayName_Lambda(GetDisplayName)
+		.GetItemToolTip_Lambda(GetToolTip)
+		.Items(TArray<EAutomatedPerfTestType>(
+			{
+				EAutomatedPerfTestType::Sequence,
+				EAutomatedPerfTestType::Replay ,
+				EAutomatedPerfTestType::ProfileGo ,
+				EAutomatedPerfTestType::StaticCamera ,
+				EAutomatedPerfTestType::Material
+			})
+		));
+
+	AddTestNodeOptions(TreeNode);
+}
+
+void FAutomatedPerfTestLaunchExtensionInstance::CustomizeAutomatedTestCommandLine(FString& InOutCommandLine)
+{
+	// Most of the parameters here are defined and handled in the corresponding 
+	// Gauntlet controller which launches the automated perf test with the required 
+	// params based on the options passed down here.
+
+	if (GetConfigBool(EConfig::User_PerProfile, EnableCSVProfilerConfig, true))
+	{
+		InOutCommandLine += TEXT(" -AutomatedPerfTest.DoCSVProfiler");
+	}
+
+	const auto AddStringParamConfig = 
+	[this, &InOutCommandLine](const FString& Config, const FString& Param, const FString& DefaultVal = "")
+	{
+		FString ConfigVal = GetConfigString(EConfig::PerProfile, *Config, *DefaultVal);
+		if (!ConfigVal.IsEmpty())
+		{
+			InOutCommandLine += FString::Printf(TEXT(" -%s=%s"), *Param, *ConfigVal);
+		}
+	};
+
+	const int32 IterationCount = GetConfigInteger(EConfig::User_PerProfile, IterationsConfig, DefaultIterationCount);
+	InOutCommandLine += FString::Printf(TEXT(" -iterations=%u"), IterationCount);
+
+	switch (CurrentTestType)
+	{
+	case EAutomatedPerfTestType::ProfileGo: 
+	{
+	    FString ConfigPath = GetConfigString(EConfig::PerProfile, ProfileGoConfigFileName);
+		if(ConfigPath.IsEmpty())
+		{
+			ConfigPath = FPaths::Combine(FPaths::ProjectDir(), "Saved", "Profiling", "ProfileGo.json");
+		}
+
+		if (FPaths::FileExists(ConfigPath))
+		{
+			InOutCommandLine += FString::Printf(TEXT(" -profilego.config=%s"), *ConfigPath);
+		}
+
+		AddStringParamConfig(ProfileGoMapConfig, TEXT("Map"));
+		break;
+	}
+	case EAutomatedPerfTestType::Replay:
+	{
+		AddStringParamConfig(ReplayConfig, TEXT("AutomatedPerfTest.ReplayPerfTest.ReplayName"));
+		break;
+	}
+	case EAutomatedPerfTestType::Sequence:
+	{
+		AddStringParamConfig(SequenceTestNameConfig, TEXT("AutomatedPerfTest.SequencePerfTest.MapSequenceName"));
+		break;
+	}
+	case EAutomatedPerfTestType::Material:
+	{
+		// Material test uses settings for test parameters.
+		break;
+	}
+	case EAutomatedPerfTestType::StaticCamera:
+	{
+		AddStringParamConfig(StaticCameraMapConfig, TEXT("AutomatedPerfTest.StaticCameraPerfTest.MapName"));
+		break;
+	}
+	default: break;
+	}
+
+	if (GetConfigBool(EConfig::PerProfile, WindowedConfig))
+	{
+		InOutCommandLine += TEXT(" -windowed");
+	}
+}
+
+void FAutomatedPerfTestLaunchExtensionInstance::ExportProfileGoScenarios(const FString& Filename)
+{
+	// Exporting ProfileGo Scenarios is only supported in editor 
+	// as we will not have engine context otherwise.
+#if WITH_EDITOR
+	UWorld* World = GWorld;
+	UEditorEngine* EditorEngine = Cast<UEditorEngine>(GEngine);
+	if (GIsEditor && EditorEngine != nullptr && EditorEngine->PlayWorld != nullptr)
+	{
+		World = EditorEngine->PlayWorld.Get();
+	}
+
+	if (World)
+	{
+		// TODO: This needs to be standardized and/or made configurable with error handling. 
+		const FString OutputPath = FPaths::Combine(FPaths::ProjectDir(), "Saved", "Profiling", Filename);
+
+		UProfileGo& ProfileGo = UProfileGo::GetCDO();
+		ProfileGo.AddScenariosInLevel(World);
+		ProfileGo.SaveToJSON(OutputPath);
+	}
+#endif // WITH_EDITOR
+}
+
+void FAutomatedPerfTestLaunchExtensionInstance::AddTestNodeOptions(ProjectLauncher::FLaunchProfileTreeNode& TreeNode)
+{
+	AddSequenceTestNodeOptions(TreeNode);
+	AddReplayTestNodeOptions(TreeNode);
+	AddStaticCameraTestNodeOptions(TreeNode);
+	AddMaterialTestNodeOptions(TreeNode);
+	AddProfileGoTestNodeOptions(TreeNode);
+}
+
+void FAutomatedPerfTestLaunchExtensionInstance::AddSequenceTestNodeOptions(ProjectLauncher::FLaunchProfileTreeNode& TreeNode)
+{
+	ProjectLauncher::FLaunchProfileTreeNode::FCallbacks&& Callbacks =
+		GetTestTypeCallbacks<EAutomatedPerfTestType::Sequence>();
+
+	TArray<FString> Sequences;
+
+#if WITH_EDITOR
+	const UAutomatedSequencePerfTestProjectSettings* const Settings = GetDefault<UAutomatedSequencePerfTestProjectSettings>();
+	for (const FAutomatedPerfTestMapSequenceCombo& Combo : Settings->MapsAndSequencesToTest)
+	{
+		Sequences.Add(Combo.ComboName.ToString());
+	}
+#endif // WITH_EDITOR
+
+	constexpr const TCHAR* ConfigName = SequenceTestNameConfig;
+	auto OnSequenceChanged = [this](FString Combo)
+	{
+		OnPropertyChanged();
+		SetConfigString(EConfig::PerProfile, ConfigName, Combo);
+	};
+
+	auto GetSequence = [this]() { return GetConfigString(EConfig::PerProfile, ConfigName); };
+	auto GetDisplayName = [this](FString Combo) -> FText { return FText::FromString(GetConfigString(EConfig::PerProfile, ConfigName, *Combo)); };
+
+	// TODO: In non-editor builds, we should make this a string input box
+	// so that end-users can still run perf tests even if the settings are 
+	// not available.
+	TreeNode.AddWidget(LOCTEXT("SequenceComboNameLabel", "Sequence Combo Name"),
+		MoveTemp(Callbacks),
+		SNew(SCustomLaunchCombo<FString>)
+		.OnSelectionChanged_Lambda(MoveTemp(OnSequenceChanged))
+		.SelectedItem_Lambda(MoveTemp(GetSequence))
+		.GetDisplayName_Lambda(MoveTemp(GetDisplayName))
+		.Items(Sequences)
+	);
+}
+
+void FAutomatedPerfTestLaunchExtensionInstance::AddReplayTestNodeOptions(ProjectLauncher::FLaunchProfileTreeNode& TreeNode)
+{
+	ProjectLauncher::FLaunchProfileTreeNode::FCallbacks&& Callbacks =
+		GetTestTypeCallbacks<EAutomatedPerfTestType::Replay>();
+
+	const TCHAR* FileTypeFilter = TEXT("Replay files (*.replay)|*.replay");
+	constexpr const TCHAR* ConfigName = ReplayConfig;
+	TreeNode.AddFileString(LOCTEXT("ReplayTestNameLabel", "Replay Name"), 
+	{
+		.GetValue = [this]() { return GetConfigString(EConfig::PerProfile, ConfigName); },
+		.SetValue = [this](FString Value) { SetConfigString(EConfig::PerProfile, ConfigName, Value); },
+		.IsVisible = Callbacks.IsVisible,
+		.IsEnabled = Callbacks.IsEnabled
+	}, FileTypeFilter);
+}
+
+void FAutomatedPerfTestLaunchExtensionInstance::AddProfileGoTestNodeOptions(ProjectLauncher::FLaunchProfileTreeNode& TreeNode)
+{
+    ProjectLauncher::FLaunchProfileTreeNode::FCallbacks&& Callbacks = 
+		GetTestTypeCallbacks<EAutomatedPerfTestType::ProfileGo>();
+
+	constexpr const TCHAR* ConfigName = ProfileGoMapConfig;
+	constexpr const TCHAR* ConfigFileName = ProfileGoConfigFileName;
+	constexpr const TCHAR* FileTypeFilter = TEXT("ProfileGo Config JSON (*.json)|*.json");
+
+	constexpr bool bIncludeNonContentDirMaps = true;
+	const FString ProjectPath = GetProfile()->GetProjectBasePath();
+	TArray<FString> Maps = GetModel()->GetAvailableProjectMapNames(ProjectPath, bIncludeNonContentDirMaps);
+
+	auto OnMapChanged = [this](FString Map)
+	{
+		OnPropertyChanged();
+		SetConfigString(EConfig::PerProfile, ConfigName, Map);
+	};
+
+	auto GetMap = [this]() { return GetConfigString(EConfig::PerProfile, ConfigName); };
+	auto GetDisplayName = [this](FString Map) -> FText { return FText::FromString(*Map); };
+
+	TreeNode.AddWidget(LOCTEXT("TestExportLabel", "Export ProfileGo"),
+		MoveTemp(Callbacks),
+		SNew(SButton)
+		.HAlign(HAlign_Center)
+		.VAlign(VAlign_Center)
+		.Text(LOCTEXT("TestExportLabelButton", "Export JSON"))
+		.OnClicked_Lambda([this]() -> FReply
+		{
+			const TCHAR* ProfileGoFilename = TEXT("ProfileGo.json");
+			ExportProfileGoScenarios(ProfileGoFilename);
+			return FReply::Handled();
+		})
+		.IsEnabled_Lambda([this]() { return GetTestType() == EAutomatedPerfTestType::ProfileGo && IsEditorBuild(); })
+	)
+	.AddFileString(LOCTEXT("ProfileGoConfigLabel","Config File"),
+	{
+		.GetValue = [this]() { return GetConfigString(EConfig::PerProfile, ConfigFileName); },
+		.SetValue = [this](FString Value) { SetConfigString(EConfig::PerProfile, ConfigFileName, Value); },
+		.IsVisible = Callbacks.IsVisible,
+		.IsEnabled = Callbacks.IsEnabled
+	}, 
+	FileTypeFilter)
+	.AddWidget(LOCTEXT("ProfileGoMapLabel", "Map"),
+		MoveTemp(Callbacks),
+		SNew(SCustomLaunchCombo<FString>)
+		.OnSelectionChanged_Lambda(MoveTemp(OnMapChanged))
+		.SelectedItem_Lambda(MoveTemp(GetMap))
+		.GetDisplayName_Lambda(MoveTemp(GetDisplayName))
+		.Items(MoveTemp(Maps))
+	);
+}
+
+void FAutomatedPerfTestLaunchExtensionInstance::AddStaticCameraTestNodeOptions(ProjectLauncher::FLaunchProfileTreeNode& TreeNode)
+{
+	ProjectLauncher::FLaunchProfileTreeNode::FCallbacks&& Callbacks =
+		GetTestTypeCallbacks<EAutomatedPerfTestType::StaticCamera>();
+
+	constexpr const TCHAR* ConfigName = StaticCameraMapConfig;
+	TArray<FString> Maps;
+
+#if WITH_EDITOR
+	const UAutomatedStaticCameraPerfTestProjectSettings* const Settings = GetDefault<UAutomatedStaticCameraPerfTestProjectSettings>();
+	for (const FSoftObjectPath& Map : Settings->MapsToTest)
+	{
+		Maps.Add(Map.ToString());
+	}
+#endif // WITH_EDITOR
+
+	auto OnMapChanged = [this](FString Map)
+	{
+		OnPropertyChanged();
+		SetConfigString(EConfig::PerProfile, ConfigName, Map);
+	};
+
+	auto GetMap = [this]() { return GetConfigString(EConfig::PerProfile, ConfigName); };
+	auto GetDisplayName = [this](FString Map) -> FText { return FText::FromString(GetConfigString(EConfig::PerProfile, ConfigName, *Map)); };
+
+	TreeNode.AddWidget(LOCTEXT("StaticCameraTestNameLabel", "Static Camera Map Name"), 
+		 MoveTemp(Callbacks),
+		 SNew(SCustomLaunchCombo<FString>)
+		.OnSelectionChanged_Lambda(MoveTemp(OnMapChanged))
+		.SelectedItem_Lambda(MoveTemp(GetMap))
+		.GetDisplayName_Lambda(MoveTemp(GetDisplayName))
+		.Items(MoveTemp(Maps))
+	);
+}
+
+void FAutomatedPerfTestLaunchExtensionInstance::AddMaterialTestNodeOptions(ProjectLauncher::FLaunchProfileTreeNode& TreeNode)
+{
+	// Material test does not have specific params.
+}
+
+void FAutomatedPerfTestLaunchExtensionInstance::OnTestTypeSelectionChanged(EAutomatedPerfTestType Type)
+{
+	SetTestType(Type);
+	OnPropertyChanged();
+}
+
+void FAutomatedPerfTestLaunchExtensionInstance::SetTestType(EAutomatedPerfTestType Type)
+{
+	CurrentTestType = Type;
+	SetConfigInteger(EConfig::PerProfile, TEXT("TestType"), (int32)Type);
+}
+
+EAutomatedPerfTestType FAutomatedPerfTestLaunchExtensionInstance::GetTestType() const
+{
+	return (EAutomatedPerfTestType)GetConfigInteger(EConfig::PerProfile, TEXT("TestType"), (int32)CurrentTestType);
+}
+
+TSharedPtr<ProjectLauncher::FLaunchExtensionInstance> FAutomatedPerfTestLaunchExtension::CreateInstanceForProfile(ProjectLauncher::FLaunchExtensionInstance::FArgs& InArgs)
+{
+	return MakeShared<FAutomatedPerfTestLaunchExtensionInstance>(InArgs);
+}
+
+const TCHAR* FAutomatedPerfTestLaunchExtension::GetInternalName() const
+{
+	return TEXT("AutomatedPerfTest");
+}
+
+FText FAutomatedPerfTestLaunchExtension::GetDisplayName() const
+{
+	return LOCTEXT("ExtensionName", "Automated Perf Test");
+}
+
+
+#undef LOCTEXT_NAMESPACE
\ No newline at end of file
diff --git a/Engine/Plugins/Developer/ProjectLauncher/Source/CommonLaunchExtensions/Private/AutomatedPerfTest/AutomatedPerfTestLaunchExtension.h b/Engine/Plugins/Developer/ProjectLauncher/Source/CommonLaunchExtensions/Private/AutomatedPerfTest/AutomatedPerfTestLaunchExtension.h
new file mode 100644
--- /dev/null
+++ b/Engine/Plugins/Developer/ProjectLauncher/Source/CommonLaunchExtensions/Private/AutomatedPerfTest/AutomatedPerfTestLaunchExtension.h
@@ -0,0 +1,78 @@
+// Copyright Epic Games, Inc. All Rights Reserved.
+
+#pragma once
+
+#include "Extension/LaunchExtension.h"
+#include "Templates/UnrealTemplate.h"
+
+enum class EAutomatedPerfTestType : uint8
+{
+	Sequence,
+	Replay,
+	ProfileGo,
+	StaticCamera,
+	Material,
+	MAX
+};
+
+class FAutomatedPerfTestLaunchExtensionInstance : public ProjectLauncher::FAutomatedTestLaunchExtensionInstance
+{
+	using Super = ProjectLauncher::FAutomatedTestLaunchExtensionInstance;
+public:
+	FAutomatedPerfTestLaunchExtensionInstance(FArgs& InArgs) : FAutomatedTestLaunchExtensionInstance(InArgs) {};
+	virtual ~FAutomatedPerfTestLaunchExtensionInstance() = default;
+
+	virtual bool HasCustomExtensionMenu() const { return true; }
+	virtual void MakeCustomExtensionSubmenu(FMenuBuilder& MenuBuilder) override;
+
+	virtual void OnPropertyChanged() override;
+	virtual void OnTestAdded(ILauncherProfileAutomatedTestRef AutomatedTest) override;
+	virtual const FString GetTestInternalName() const override;
+
+	virtual void CustomizeTree(ProjectLauncher::FLaunchProfileTreeData& ProfileTreeData) override;
+	virtual void CustomizeAutomatedTestCommandLine(FString& InOutCommandLine) override;
+
+private:
+	void ExportProfileGoScenarios(const FString& Filename);
+
+	void AddTestNodeOptions(ProjectLauncher::FLaunchProfileTreeNode& TreeNode);
+	void AddSequenceTestNodeOptions(ProjectLauncher::FLaunchProfileTreeNode& TreeNode);
+	void AddReplayTestNodeOptions(ProjectLauncher::FLaunchProfileTreeNode& TreeNode);
+	void AddProfileGoTestNodeOptions(ProjectLauncher::FLaunchProfileTreeNode& TreeNode);
+	void AddStaticCameraTestNodeOptions(ProjectLauncher::FLaunchProfileTreeNode& TreeNode);
+	void AddMaterialTestNodeOptions(ProjectLauncher::FLaunchProfileTreeNode& TreeNode);
+
+	void OnTestTypeSelectionChanged(EAutomatedPerfTestType Type);
+	void SetTestType(EAutomatedPerfTestType Type);
+	EAutomatedPerfTestType GetTestType() const;
+
+	template<EAutomatedPerfTestType TestType>
+	ProjectLauncher::FLaunchProfileTreeNode::FCallbacks GetTestTypeCallbacks()
+	{
+		const auto IsTestType = [this]()
+		{
+			const bool bEnabled = IsTestActive() || GetConfigBool(EConfig::PerProfile, TEXT("Enabled"));
+			return CurrentTestType == TestType && bEnabled;
+		};
+
+		ProjectLauncher::FLaunchProfileTreeNode::FCallbacks Callbacks
+		{
+			.IsVisible = IsTestType,
+			.IsEnabled = IsTestType
+		};
+
+		return MoveTemp(Callbacks);
+	}
+
+	EAutomatedPerfTestType CurrentTestType = EAutomatedPerfTestType::Sequence;
+	bool bEnableCSVProfiler = false;
+};
+
+
+class FAutomatedPerfTestLaunchExtension : public ProjectLauncher::FLaunchExtension
+{
+public:
+	virtual TSharedPtr<ProjectLauncher::FLaunchExtensionInstance> CreateInstanceForProfile(ProjectLauncher::FLaunchExtensionInstance::FArgs& InArgs) override;
+	virtual const TCHAR* GetInternalName() const override;
+	virtual FText GetDisplayName() const override;
+};
\ No newline at end of file
diff --git a/Engine/Plugins/Developer/ProjectLauncher/Source/CommonLaunchExtensions/Private/CommonLaunchExtensionsModule.cpp b/Engine/Plugins/Developer/ProjectLauncher/Source/CommonLaunchExtensions/Private/CommonLaunchExtensionsModule.cpp
--- a/Engine/Plugins/Developer/ProjectLauncher/Source/CommonLaunchExtensions/Private/CommonLaunchExtensionsModule.cpp
+++ b/Engine/Plugins/Developer/ProjectLauncher/Source/CommonLaunchExtensions/Private/CommonLaunchExtensionsModule.cpp
@@ -6,6 +6,7 @@
 #include "Globals/GlobalsLaunchExtension.h"
 #include "Insights/InsightsLaunchExtension.h"
 #include "BootTest/BootTestLaunchExtension.h"
+#include "AutomatedPerfTest/AutomatedPerfTestLaunchExtension.h"
 
 class FCommonLaunchExtensionsModule : public IModuleInterface
 {
@@ -16,10 +17,12 @@ public:
 		Globals = MakeShared<FGlobalsLaunchExtension>();
 		Insights = MakeShared<FInsightsLaunchExtension>();
 		BootTest = MakeShared<FBootTestLaunchExtension>();
+		AutomatedPerfTest = MakeShared<FAutomatedPerfTestLaunchExtension>();
 
 		IProjectLauncherModule::Get().RegisterExtension(Globals.ToSharedRef());
 		IProjectLauncherModule::Get().RegisterExtension(Insights.ToSharedRef());
 		IProjectLauncherModule::Get().RegisterExtension(BootTest.ToSharedRef());
+		IProjectLauncherModule::Get().RegisterExtension(AutomatedPerfTest.ToSharedRef());
 	}
 
 	virtual void ShutdownModule() override
@@ -29,17 +32,20 @@ public:
 			ProjectLauncher->UnregisterExtension(Globals.ToSharedRef());
 			ProjectLauncher->UnregisterExtension(Insights.ToSharedRef());
 			ProjectLauncher->UnregisterExtension(BootTest.ToSharedRef());
+			ProjectLauncher->UnregisterExtension(AutomatedPerfTest.ToSharedRef());
 		}
 
 		Globals.Reset();
 		Insights.Reset();
 		BootTest.Reset();
+		AutomatedPerfTest.Reset();
 	}
 
 private:
 	TSharedPtr<FGlobalsLaunchExtension> Globals;
 	TSharedPtr<FInsightsLaunchExtension> Insights;
 	TSharedPtr<FBootTestLaunchExtension> BootTest;
+	TSharedPtr<FAutomatedPerfTestLaunchExtension> AutomatedPerfTest;
 
 };
 
diff --git a/Engine/Plugins/Developer/ProjectLauncher/Source/ProjectLauncher/Private/Model/ProjectLauncherModel.cpp b/Engine/Plugins/Developer/ProjectLauncher/Source/ProjectLauncher/Private/Model/ProjectLauncherModel.cpp
--- a/Engine/Plugins/Developer/ProjectLauncher/Source/ProjectLauncher/Private/Model/ProjectLauncherModel.cpp
+++ b/Engine/Plugins/Developer/ProjectLauncher/Source/ProjectLauncher/Private/Model/ProjectLauncherModel.cpp
@@ -746,7 +746,7 @@ namespace ProjectLauncher
 	}
 
 
-	TArray<FString> FModel::GetAndCacheMapPaths(const FString& InOptionalProjectPath) // @todo: map list parsing should ideally be asyncronous, showing a spinner in the map selector controls until its finished etc.
+	TArray<FString> FModel::GetAndCacheMapPaths(const FString& InOptionalProjectPath, bool bIncludeNonContentDirMaps) // @todo: map list parsing should ideally be asyncronous, showing a spinner in the map selector controls until its finished etc.
 	{
 		// prepare values
 		FString ProjectPath;
@@ -785,7 +785,7 @@ namespace ProjectLauncher
 					FString MapFileName = FPackageName::LongPackageNameToFilename(MapAsset.PackageName.ToString(), FPackageName::GetMapPackageExtension());
 					MapFileName = FPaths::ConvertRelativePathToFull(MapFileName);
 
-					if (FPaths::IsUnderDirectory(MapFileName, ContentDir))
+					if (FPaths::IsUnderDirectory(MapFileName, ContentDir) || bIncludeNonContentDirMaps)
 					{
 						MapFileList.Add( MapFileName );
 					}
@@ -821,9 +821,9 @@ namespace ProjectLauncher
 		}
 	}
 
-	TArray<FString> FModel::GetAvailableProjectMapNames( const FString& InProjectPath )
+	TArray<FString> FModel::GetAvailableProjectMapNames( const FString& InProjectPath, bool bIncludeNonContentDirMaps)
 	{
-		TArray<FString> Maps = GetAndCacheMapPaths(InProjectPath);
+		TArray<FString> Maps = GetAndCacheMapPaths(InProjectPath, bIncludeNonContentDirMaps);
 		for (FString& Map : Maps)
 		{
 			Map = FPaths::GetBaseFilename(Map);
@@ -831,9 +831,9 @@ namespace ProjectLauncher
 		return MoveTemp(Maps);
 	}
 
-	TArray<FString> FModel::GetAvailableProjectMapPaths( const FString& InProjectPath )
+	TArray<FString> FModel::GetAvailableProjectMapPaths( const FString& InProjectPath, bool bIncludeNonContentDirMaps)
 	{
-		return GetAndCacheMapPaths(InProjectPath);
+		return GetAndCacheMapPaths(InProjectPath, bIncludeNonContentDirMaps);
 	}
 
 	TArray<FString> FModel::GetAvailableEngineMapNames()
diff --git a/Engine/Plugins/Developer/ProjectLauncher/Source/ProjectLauncher/Private/ProfileTree/LaunchProfileTreeData.cpp b/Engine/Plugins/Developer/ProjectLauncher/Source/ProjectLauncher/Private/ProfileTree/LaunchProfileTreeData.cpp
--- a/Engine/Plugins/Developer/ProjectLauncher/Source/ProjectLauncher/Private/ProfileTree/LaunchProfileTreeData.cpp
+++ b/Engine/Plugins/Developer/ProjectLauncher/Source/ProjectLauncher/Private/ProfileTree/LaunchProfileTreeData.cpp
@@ -401,7 +401,116 @@ namespace ProjectLauncher
 	}
 
 
-	FLaunchProfileTreeNode& FLaunchProfileTreeNode::AddCommandLineString( FText InName, FStringCallbacks&& StringCallbacks )
+	FLaunchProfileTreeNode& FLaunchProfileTreeNode::AddFileString(FText InName, FStringCallbacks&& StringCallbacks, FString FileFilter)
+	{
+		// Invokes file picker dialog and passes the selected
+		// file as a string to given callback.
+
+		check(StringCallbacks.GetValue);
+		check(StringCallbacks.SetValue);
+
+		FCallbacks WidgetCallbacks
+		{
+			.IsVisible = StringCallbacks.IsVisible,
+			.IsEnabled = StringCallbacks.IsEnabled,
+		};
+
+		if (StringCallbacks.GetDefaultValue != nullptr)
+		{
+			WidgetCallbacks.IsDefault = [StringCallbacks]()
+			{
+				return StringCallbacks.GetValue() == StringCallbacks.GetDefaultValue();
+			};
+
+			WidgetCallbacks.SetToDefault = [StringCallbacks]()
+			{
+				StringCallbacks.SetValue(StringCallbacks.GetDefaultValue());
+				// nb. OnPropertyChanged is called via SCustomLaunchWidgetTableRow's set to default
+			};
+		};
+
+		auto SetString = [this, StringCallbacks]( const FText& InText, ETextCommit::Type InTextCommit )
+		{
+			StringCallbacks.SetValue( InText.ToString() );
+			TreeData->TreeBuilder->OnPropertyChanged();
+		};
+
+		auto GetString = [StringCallbacks]()
+		{
+			return FText::FromString(StringCallbacks.GetValue());
+		};
+
+		auto OnBrowse = [this,StringCallbacks, FileFilter]()
+		{
+			FString InitialDirectory = StringCallbacks.GetValue();
+			if (InitialDirectory.IsEmpty())
+			{
+				InitialDirectory = this->TreeData->Profile->GetProjectBasePath();
+			}
+			if (!InitialDirectory.IsEmpty() && FPaths::IsRelative(InitialDirectory))
+			{
+				InitialDirectory = FPaths::Combine( FPaths::RootDir(), InitialDirectory );
+			}
+
+			const void* ParentWindowHandle = FSlateApplication::Get().FindBestParentWindowHandleForDialogs(nullptr);
+			TArray<FString> OutFiles;
+			FString OutDirectory;
+
+			if (FDesktopPlatformModule::Get()->OpenFileDialog(ParentWindowHandle, LOCTEXT("FileDialogTitle", "Choose a file").ToString(), InitialDirectory, "", FileFilter, EFileDialogFlags::None, OutFiles))
+			{
+				if (OutFiles.Num() > 0)
+				{
+					FString& FilePath = OutFiles[0];
+					if (FPaths::IsUnderDirectory(FilePath, FPaths::RootDir()))
+					{
+						FPaths::MakePathRelativeTo(FilePath, *FPaths::RootDir());
+					}
+
+					StringCallbacks.SetValue(FilePath);
+					TreeData->TreeBuilder->OnPropertyChanged();
+				}
+			}
+
+			return FReply::Handled();
+		};
+
+
+		return AddWidget( InName,
+			MoveTemp(WidgetCallbacks),
+			SNew(SHorizontalBox)
+
+			// path field
+			+SHorizontalBox::Slot()
+			.FillWidth(1)
+			.VAlign(VAlign_Center)
+			[
+				SNew(SEditableTextBox)
+				.Font(FCoreStyle::Get().GetFontStyle("SmallFont"))
+				.Text_Lambda( GetString )
+				.OnTextCommitted_Lambda( SetString )
+			]
+
+			// browse button
+			+SHorizontalBox::Slot()
+			.AutoWidth()
+			.Padding(FMargin(4, 0, 0, 0))
+			.VAlign(VAlign_Center)
+			[
+				SNew(SButton)
+				.ButtonStyle(FProjectLauncherStyle::Get(), "HoverHintOnly")
+				.ToolTipText(LOCTEXT("FileBrowseTip", "Browse for a file"))
+				.OnClicked_Lambda(OnBrowse)
+				.ContentPadding(2)
+				[
+					SNew(SImage)
+					.Image(FProjectLauncherStyle::Get().GetBrush("PathPickerButton"))
+					.ColorAndOpacity(FSlateColor::UseForeground())
+				]
+			]
+		);
+	}
+
+	FLaunchProfileTreeNode& FLaunchProfileTreeNode::AddCommandLineString(FText InName, FStringCallbacks&& StringCallbacks)
 	{
 		if (!TreeData->Model->AreExtensionsEnabled())
 		{
diff --git a/Engine/Plugins/Developer/ProjectLauncher/Source/ProjectLauncher/Public/Model/ProjectLauncherModel.h b/Engine/Plugins/Developer/ProjectLauncher/Source/ProjectLauncher/Public/Model/ProjectLauncherModel.h
--- a/Engine/Plugins/Developer/ProjectLauncher/Source/ProjectLauncher/Public/Model/ProjectLauncherModel.h
+++ b/Engine/Plugins/Developer/ProjectLauncher/Source/ProjectLauncher/Public/Model/ProjectLauncherModel.h
@@ -180,8 +180,8 @@ namespace ProjectLauncher
 
 		UE_API bool AreExtensionsEnabled() const;
 
-		UE_API TArray<FString> GetAvailableProjectMapNames( const FString& InProjectPath );
-		UE_API TArray<FString> GetAvailableProjectMapPaths( const FString& InProjectPath );
+		UE_API TArray<FString> GetAvailableProjectMapNames( const FString& InProjectPath, bool bIncludeNonContentDirMaps = false);
+		UE_API TArray<FString> GetAvailableProjectMapPaths( const FString& InProjectPath, bool bIncludeNonContentDirMaps = false);
 		UE_API TArray<FString> GetAvailableEngineMapNames();
 		UE_API TArray<FString> GetAvailableEngineMapPaths();
 
@@ -202,7 +202,7 @@ namespace ProjectLauncher
 		UE_API void OnModifyAutomatedTestCommandLine( const ILauncherProfileAutomatedTestRef& InAutomatedTest, const ILauncherProfileRef& InProfile, FString& InOutCommandLine );
 
 
-		UE_API TArray<FString> GetAndCacheMapPaths(const FString& InOptionalProjectPath);
+		UE_API TArray<FString> GetAndCacheMapPaths(const FString& InOptionalProjectPath, bool bIncludeNonContentDirMaps = false);
 
 		TSharedRef<ITargetDeviceProxyManager> DeviceProxyManager;
 		TSharedRef<ILauncher> Launcher;
diff --git a/Engine/Plugins/Developer/ProjectLauncher/Source/ProjectLauncher/Public/ProfileTree/LaunchProfileTreeData.h b/Engine/Plugins/Developer/ProjectLauncher/Source/ProjectLauncher/Public/ProfileTree/LaunchProfileTreeData.h
--- a/Engine/Plugins/Developer/ProjectLauncher/Source/ProjectLauncher/Public/ProfileTree/LaunchProfileTreeData.h
+++ b/Engine/Plugins/Developer/ProjectLauncher/Source/ProjectLauncher/Public/ProfileTree/LaunchProfileTreeData.h
@@ -100,7 +100,8 @@ namespace ProjectLauncher
 			FValidation Validation;
 		};
 		UE_API FLaunchProfileTreeNode& AddString( FText InName, FStringCallbacks&& StringCallbacks );
-		UE_API FLaunchProfileTreeNode& AddDirectoryString( FText InName, FStringCallbacks&& StringCallbacks );
+		UE_API FLaunchProfileTreeNode& AddDirectoryString(FText InName, FStringCallbacks&& StringCallbacks);
+		UE_API FLaunchProfileTreeNode& AddFileString(FText InName, FStringCallbacks&& StringCallbacks, FString FileFilter = TEXT("All files (*.*)|*.*"));
 		UE_API FLaunchProfileTreeNode& AddCommandLineString( FText InName, FStringCallbacks&& StringCallbacks );
 
 
