Export profile validation constructors, allowing extensions to provide validations too.

#rnx

diff --git a/Engine/Plugins/Developer/ProjectLauncher/Source/ProjectLauncher/Public/ProfileTree/LaunchProfileTreeData.h b/Engine/Plugins/Developer/ProjectLauncher/Source/ProjectLauncher/Public/ProfileTree/LaunchProfileTreeData.h
--- a/Engine/Plugins/Developer/ProjectLauncher/Source/ProjectLauncher/Public/ProfileTree/LaunchProfileTreeData.h
+++ b/Engine/Plugins/Developer/ProjectLauncher/Source/ProjectLauncher/Public/ProfileTree/LaunchProfileTreeData.h
@@ -33,9 +33,9 @@ namespace ProjectLauncher
 	{
 	public:
 		FValidation() {}
-		FValidation( std::initializer_list<ELauncherProfileValidationErrors::Type> InAssociatedErrors );
-		FValidation( std::initializer_list<ELauncherProfileValidationErrors::Type> InAssociatedErrors, std::initializer_list<FString> InAssociatedCustomErrors );
-		FValidation( std::initializer_list<FString> InAssociatedCustomErrors );
+		UE_API FValidation( std::initializer_list<ELauncherProfileValidationErrors::Type> InAssociatedErrors );
+		UE_API FValidation( std::initializer_list<ELauncherProfileValidationErrors::Type> InAssociatedErrors, std::initializer_list<FString> InAssociatedCustomErrors );
+		UE_API FValidation( std::initializer_list<FString> InAssociatedCustomErrors );
 
 		bool IsSet() const;
 
