Project Launcher misc changes:
- add profile verification callback for tree builder
- fix execution order for content scheme property change notifications

#rnx

diff --git a/Engine/Plugins/Developer/ProjectLauncher/Source/ProjectLauncher/Private/ProfileTree/GenericProfileTreeBuilder.cpp b/Engine/Plugins/Developer/ProjectLauncher/Source/ProjectLauncher/Private/ProfileTree/GenericProfileTreeBuilder.cpp
--- a/Engine/Plugins/Developer/ProjectLauncher/Source/ProjectLauncher/Private/ProfileTree/GenericProfileTreeBuilder.cpp
+++ b/Engine/Plugins/Developer/ProjectLauncher/Source/ProjectLauncher/Private/ProfileTree/GenericProfileTreeBuilder.cpp
@@ -81,6 +81,12 @@ namespace ProjectLauncher
 			return FString();
 		};
 
+		Profile->OnCustomValidation().AddRaw( this, &FGenericProfileTreeBuilder::OnValidateProfile );
+	}
+	
+	FGenericProfileTreeBuilder::~FGenericProfileTreeBuilder()
+	{
+		Profile->OnCustomValidation().RemoveAll( this );
 	}
 
 	void FGenericProfileTreeBuilder::Construct()
@@ -603,6 +609,10 @@ namespace ProjectLauncher
 
 
 
+
+	void FGenericProfileTreeBuilder::OnValidateProfile()
+	{
+	}
 
 
 	void FGenericProfileTreeBuilder::OnPropertyChanged()
@@ -858,10 +868,11 @@ namespace ProjectLauncher
 	{
 		ELauncherProfileDeploymentModes::Type DeploymentMode = GetDeployToDevice() ? ELauncherProfileDeploymentModes::CopyToDevice : ELauncherProfileDeploymentModes::DoNotDeploy;
 		Model->SetProfileContentScheme(InContentScheme, Profile, bShouldCook, DeploymentMode );
-		OnPropertyChanged();
 
 		// refresh the cached content scheme in case the option that was selected is not available
 		ContentScheme = Model->DetermineProfileContentScheme(Profile);
+
+		OnPropertyChanged();
 	}
 
 	bool FGenericProfileTreeBuilder::IsContentSchemeAvailable(EContentScheme InContentScheme, FText& OutReason) const
diff --git a/Engine/Plugins/Developer/ProjectLauncher/Source/ProjectLauncher/Public/ProfileTree/GenericProfileTreeBuilder.h b/Engine/Plugins/Developer/ProjectLauncher/Source/ProjectLauncher/Public/ProfileTree/GenericProfileTreeBuilder.h
--- a/Engine/Plugins/Developer/ProjectLauncher/Source/ProjectLauncher/Public/ProfileTree/GenericProfileTreeBuilder.h
+++ b/Engine/Plugins/Developer/ProjectLauncher/Source/ProjectLauncher/Public/ProfileTree/GenericProfileTreeBuilder.h
@@ -28,7 +28,7 @@ namespace ProjectLauncher
 	{
 	public:
 		UE_API FGenericProfileTreeBuilder( const ILauncherProfileRef& Profile, const ILauncherProfileRef& InDefaultProfile, const TSharedRef<FModel>& InModel );
-		virtual ~FGenericProfileTreeBuilder() = default;
+		virtual ~FGenericProfileTreeBuilder();
 
 		UE_API virtual void Construct() override;
 
@@ -173,6 +173,7 @@ namespace ProjectLauncher
 	protected:
 
 		UE_API virtual void OnPropertyChanged() override;
+		UE_API virtual void OnValidateProfile();
 		UE_API void RefreshContentScheme();
 		UE_API void CacheStartupMapList() const;
 		UE_API void CacheArchitectures();
