Zen workflows are not available when IoStore is disabled
- new profiles default to overriding it to true (because they default to using zen streaming)
- override property is only visible if the property is disabled in the project settings

...still not ideal because we don't need the IoStore UAT parameter if it's specified in the project settings

diff --git a/Engine/Plugins/Developer/ProjectLauncher/Source/ProjectLauncher/Private/Model/ProjectLauncherModel.cpp b/Engine/Plugins/Developer/ProjectLauncher/Source/ProjectLauncher/Private/Model/ProjectLauncherModel.cpp
--- a/Engine/Plugins/Developer/ProjectLauncher/Source/ProjectLauncher/Private/Model/ProjectLauncherModel.cpp
+++ b/Engine/Plugins/Developer/ProjectLauncher/Source/ProjectLauncher/Private/Model/ProjectLauncherModel.cpp
@@ -595,6 +595,7 @@ namespace ProjectLauncher
 		// read project packaging settings
 		const TCHAR* ProjectPackagingConfigSection = TEXT("/Script/UnrealEd.ProjectPackagingSettings");
 		InConfig.GetBool(ProjectPackagingConfigSection, TEXT("bUseZenStore"), OutResult.bUseZenStore, GGameIni);
+		InConfig.GetBool(ProjectPackagingConfigSection, TEXT("bUseIoStore"), OutResult.bUseIoStore, GGameIni);
 		InConfig.GetBool(ProjectPackagingConfigSection, TEXT("bEnablePakStreaming"), OutResult.bHasAutomaticZenPakStreamingWorkspaceCreation, GGameIni);
 		InConfig.GetString(ProjectPackagingConfigSection, TEXT("BuildTarget"), OutResult.DefaultBuildTargetName, GGameIni);
 		if (OutResult.DefaultBuildTargetName.IsEmpty())
@@ -647,6 +648,7 @@ namespace ProjectLauncher
 			// apply any in-memory properties @fixme: is this necessary?
 			const UProjectPackagingSettings* ProjectPackagingSettings = UProjectPackagingSettings::StaticClass()->GetDefaultObject<UProjectPackagingSettings>();
 			Result.bUseZenStore = ProjectPackagingSettings->bUseZenStore;
+			Result.bUseIoStore = ProjectPackagingSettings->bUseIoStore;
 			if (!ProjectPackagingSettings->BuildTarget.IsEmpty())
 			{
 				Result.DefaultBuildTargetName = ProjectPackagingSettings->BuildTarget;
diff --git a/Engine/Plugins/Developer/ProjectLauncher/Source/ProjectLauncher/Private/ProfileTree/GenericProfileTreeBuilder.cpp b/Engine/Plugins/Developer/ProjectLauncher/Source/ProjectLauncher/Private/ProfileTree/GenericProfileTreeBuilder.cpp
--- a/Engine/Plugins/Developer/ProjectLauncher/Source/ProjectLauncher/Private/ProfileTree/GenericProfileTreeBuilder.cpp
+++ b/Engine/Plugins/Developer/ProjectLauncher/Source/ProjectLauncher/Private/ProfileTree/GenericProfileTreeBuilder.cpp
@@ -41,6 +41,11 @@ namespace ProjectLauncher
 			return (ContentScheme == EContentScheme::PakFiles || ContentScheme == EContentScheme::DevelopmentPackage);
 		};
 
+		ForZen = [this]()
+		{
+			return (ContentScheme == EContentScheme::ZenStreaming || ContentScheme == EContentScheme::ZenPakStreaming);
+		};
+
 		ForZenWS = [this]()
 		{
 			return (ContentScheme == EContentScheme::ZenPakStreaming);
@@ -228,7 +233,8 @@ namespace ProjectLauncher
 	{
 		HeadingNode.AddWidget( LOCTEXT("ContentSchemeLabel", "Content Scheme"), 
 			{
-				.Validation = FValidation( {ELauncherProfileValidationErrors::CookOnTheFlyDoesntSupportServer, ELauncherProfileValidationErrors::ShippingDoesntSupportCommandlineOptionsCantUseCookOnTheFly, ELauncherProfileValidationErrors::ZenPakStreamingRequiresDeployAndLaunch} ),
+				.Validation = FValidation( {ELauncherProfileValidationErrors::CookOnTheFlyDoesntSupportServer, ELauncherProfileValidationErrors::ShippingDoesntSupportCommandlineOptionsCantUseCookOnTheFly, ELauncherProfileValidationErrors::ZenPakStreamingRequiresDeployAndLaunch}, 
+											{TEXT("Validation_ZenNeedsIoStore")} ),
 			},
 			SNew(SCustomLaunchContentSchemeCombo)
 			.OnSelectionChanged(this, &FGenericProfileTreeBuilder::SetContentScheme)
@@ -254,13 +260,17 @@ namespace ProjectLauncher
 
 	void FGenericProfileTreeBuilder::AddUseIoStoreProperty( FLaunchProfileTreeNode& HeadingNode )
 	{
+		auto IoStorePropertyVisible = [this]()
+		{
+			return ForPak()  && !Model->GetProjectSettings(Profile).bUseIoStore;
+		};
+
 		HeadingNode.AddBoolean( LOCTEXT("UseIoStoreLabel", "Use Io Store"),
 			{
 				.GetValue = [this]()			{ return Profile->IsUsingIoStore(); },
 				.SetValue = [this](bool bValue)	{ Profile->SetUseIoStore(bValue); },
 				.GetDefaultValue = [this]()		{ return DefaultProfile->IsUsingIoStore(); },
-				.IsVisible = ForPak,
-				.Validation = FValidation({ELauncherProfileValidationErrors::IoStoreRequiresPakFiles}),
+				.IsVisible = IoStorePropertyVisible,
 			}
 		);
 	}
@@ -612,6 +622,11 @@ namespace ProjectLauncher
 
 	void FGenericProfileTreeBuilder::OnValidateProfile()
 	{
+		// verify that IoStore is available when a Zen Workflow is in use
+		if ((ContentScheme == EContentScheme::ZenStreaming || ContentScheme == EContentScheme::ZenPakStreaming) && !Model->GetProjectSettings(Profile).bUseIoStore)
+		{
+			Profile->AddCustomError(TEXT("Validation_ZenNeedsIoStore"), LOCTEXT("Error_ZenNeedsIoStore", "Zen Workflows require IoStore to be enabled in your project's Packaging Settings. It cannot be overridden in a profile"));
+		}
 	}
 
 
diff --git a/Engine/Plugins/Developer/ProjectLauncher/Source/ProjectLauncher/Public/Model/ProjectLauncherModel.h b/Engine/Plugins/Developer/ProjectLauncher/Source/ProjectLauncher/Public/Model/ProjectLauncherModel.h
--- a/Engine/Plugins/Developer/ProjectLauncher/Source/ProjectLauncher/Public/Model/ProjectLauncherModel.h
+++ b/Engine/Plugins/Developer/ProjectLauncher/Source/ProjectLauncher/Public/Model/ProjectLauncherModel.h
@@ -69,6 +69,9 @@ namespace ProjectLauncher
 		// whether Zen Store is enabled via the project settings - in this case, we can't use Loose Files
 		bool bUseZenStore = false;
 
+		// whether Io Store is enabled via the project settings - we can only use zen workflows when this is enabled
+		bool bUseIoStore = false;
+
 		// whether a workspace will be created automatically when using Zen Pak streaming
 		bool bHasAutomaticZenPakStreamingWorkspaceCreation = false;
 
diff --git a/Engine/Plugins/Developer/ProjectLauncher/Source/ProjectLauncher/Public/ProfileTree/GenericProfileTreeBuilder.h b/Engine/Plugins/Developer/ProjectLauncher/Source/ProjectLauncher/Public/ProfileTree/GenericProfileTreeBuilder.h
--- a/Engine/Plugins/Developer/ProjectLauncher/Source/ProjectLauncher/Public/ProfileTree/GenericProfileTreeBuilder.h
+++ b/Engine/Plugins/Developer/ProjectLauncher/Source/ProjectLauncher/Public/ProfileTree/GenericProfileTreeBuilder.h
@@ -75,6 +75,7 @@ namespace ProjectLauncher
 
 		// helper callbacks to simplify control enable/visibility
 		FLaunchProfileTreeNode::FGetBool ForPak;
+		FLaunchProfileTreeNode::FGetBool ForZen;
 		FLaunchProfileTreeNode::FGetBool ForZenWS;
 		FLaunchProfileTreeNode::FGetBool ForCooked;
 		FLaunchProfileTreeNode::FGetBool ForEnabledCooked;
