Project Launcher re-validates the profile when the tab becomes active
- catches cases when the user has adjusted the project's Packaging Settings, affecting validation

#rnx

diff --git a/Engine/Plugins/Developer/ProjectLauncher/Source/ProjectLauncher/Private/ProjectLauncherModule.cpp b/Engine/Plugins/Developer/ProjectLauncher/Source/ProjectLauncher/Private/ProjectLauncherModule.cpp
--- a/Engine/Plugins/Developer/ProjectLauncher/Source/ProjectLauncher/Private/ProjectLauncherModule.cpp
+++ b/Engine/Plugins/Developer/ProjectLauncher/Source/ProjectLauncher/Private/ProjectLauncherModule.cpp
@@ -88,8 +88,22 @@ private:
 
 		TSharedRef<ProjectLauncher::FModel> Model = MakeShared<ProjectLauncher::FModel>(TargetDeviceServicesModule.GetDeviceProxyManager(), LauncherServicesModule.CreateLauncher(), LauncherServicesModule.GetProfileManager() );
 		
+		auto OnTabActivated = [Model](TSharedRef<SDockTab>, ETabActivationCause Cause)
+		{
+			if (Cause == ETabActivationCause::UserClickedOnTab)
+			{
+				ILauncherProfilePtr Profile = Model->GetSelectedProfile();
+				if (Profile.IsValid())
+				{
+					Profile->RefreshCustomWarningsAndErrors();
+				}
+			}
+		};
+
+
 		return SNew(SDockTab)
 		.TabRole(ETabRole::NomadTab)
+		.OnTabActivated_Lambda(OnTabActivated)
 		[
 			SNew(SProjectLauncher, Model)
 		];
