Project Launcher headings require an internal name
- prevents duplication
- allows extensions to extend existing sections, or share a section

#rnx

diff --git a/Engine/Plugins/Developer/ProjectLauncher/Source/CommonLaunchExtensions/Private/AutomatedPerfTest/AutomatedPerfTestLaunchExtension.cpp b/Engine/Plugins/Developer/ProjectLauncher/Source/CommonLaunchExtensions/Private/AutomatedPerfTest/AutomatedPerfTestLaunchExtension.cpp
--- a/Engine/Plugins/Developer/ProjectLauncher/Source/CommonLaunchExtensions/Private/AutomatedPerfTest/AutomatedPerfTestLaunchExtension.cpp
+++ b/Engine/Plugins/Developer/ProjectLauncher/Source/CommonLaunchExtensions/Private/AutomatedPerfTest/AutomatedPerfTestLaunchExtension.cpp
@@ -194,8 +194,8 @@ void FAutomatedPerfTestLaunchExtensionInstance::CustomizeTree(ProjectLauncher::F
 	// extensible in the future.
 
 	constexpr EConfig ConfigTypeCommon = EConfig::User_PerProfile;
-	FLaunchProfileTreeNode& TreeNode = ProfileTreeData
-		.AddHeading(GetExtension()->GetDisplayName())
+	FLaunchProfileTreeNode& TreeNode = 
+		AddDefaultHeading(ProfileTreeData)
 		.AddBoolean(LOCTEXT("APTEnableCSVLabel", "Enable CSV Profiler"), 
 		{
 			.GetValue = [this]() { return GetConfigBool(ConfigTypeCommon, EnableCSVProfilerConfig, true); },
diff --git a/Engine/Plugins/Developer/ProjectLauncher/Source/CommonLaunchExtensions/Private/BootTest/BootTestLaunchExtension.cpp b/Engine/Plugins/Developer/ProjectLauncher/Source/CommonLaunchExtensions/Private/BootTest/BootTestLaunchExtension.cpp
--- a/Engine/Plugins/Developer/ProjectLauncher/Source/CommonLaunchExtensions/Private/BootTest/BootTestLaunchExtension.cpp
+++ b/Engine/Plugins/Developer/ProjectLauncher/Source/CommonLaunchExtensions/Private/BootTest/BootTestLaunchExtension.cpp
@@ -17,7 +17,7 @@ void FBootTestLaunchExtensionInstance::CustomizeTree( ProjectLauncher::FLaunchPr
 		return IsTestActive();
 	};
 
-	ProfileTreeData.AddHeading(GetExtension()->GetDisplayName())
+	AddDefaultHeading(ProfileTreeData)
 		.AddBoolean(LOCTEXT("WindowedLabel","Windowed"),
 		{
 			.GetValue = [this]()          { return GetConfigBool(EConfig::PerProfile, TEXT("Windowed"));},
diff --git a/Engine/Plugins/Developer/ProjectLauncher/Source/ProjectLauncher/Private/Extension/LaunchExtension.cpp b/Engine/Plugins/Developer/ProjectLauncher/Source/ProjectLauncher/Private/Extension/LaunchExtension.cpp
--- a/Engine/Plugins/Developer/ProjectLauncher/Source/ProjectLauncher/Private/Extension/LaunchExtension.cpp
+++ b/Engine/Plugins/Developer/ProjectLauncher/Source/ProjectLauncher/Private/Extension/LaunchExtension.cpp
@@ -639,6 +639,16 @@ namespace ProjectLauncher
 		return Extension->GetDisplayName();
 	}
 
+	FLaunchProfileTreeNode& FLaunchExtensionInstance::AddDefaultHeading(ProjectLauncher::FLaunchProfileTreeData& ProfileTreeData) const
+	{
+		return AddDefaultHeading(ProfileTreeData, Extension->GetDisplayName());
+	}
+
+	FLaunchProfileTreeNode& FLaunchExtensionInstance::AddDefaultHeading(ProjectLauncher::FLaunchProfileTreeData& ProfileTreeData, FText OverrideDisplayName) const
+	{
+		return ProfileTreeData.AddHeading(Extension->GetInternalName(), OverrideDisplayName);
+	}
+
 
 
 
diff --git a/Engine/Plugins/Developer/ProjectLauncher/Source/ProjectLauncher/Private/ProfileTree/BasicProfileTreeBuilder.cpp b/Engine/Plugins/Developer/ProjectLauncher/Source/ProjectLauncher/Private/ProfileTree/BasicProfileTreeBuilder.cpp
--- a/Engine/Plugins/Developer/ProjectLauncher/Source/ProjectLauncher/Private/ProfileTree/BasicProfileTreeBuilder.cpp
+++ b/Engine/Plugins/Developer/ProjectLauncher/Source/ProjectLauncher/Private/ProfileTree/BasicProfileTreeBuilder.cpp
@@ -16,7 +16,7 @@ namespace ProjectLauncher
 	{
 		FGenericProfileTreeBuilder::Construct();
 
-		FLaunchProfileTreeNode& GeneralSettingsHeader = TreeData->AddHeading( LOCTEXT("GeneralSettingsHeading", "General Settings") );
+		FLaunchProfileTreeNode& GeneralSettingsHeader = TreeData->AddHeading( TEXT("GeneralSettings"), LOCTEXT("GeneralSettingsHeading", "General Settings") );
 		AddProjectProperty(GeneralSettingsHeader);
 		AddTargetProperty(GeneralSettingsHeader);
 		AddConfigurationProperty(GeneralSettingsHeader);
diff --git a/Engine/Plugins/Developer/ProjectLauncher/Source/ProjectLauncher/Private/ProfileTree/CustomProfileTreeBuilder.cpp b/Engine/Plugins/Developer/ProjectLauncher/Source/ProjectLauncher/Private/ProfileTree/CustomProfileTreeBuilder.cpp
--- a/Engine/Plugins/Developer/ProjectLauncher/Source/ProjectLauncher/Private/ProfileTree/CustomProfileTreeBuilder.cpp
+++ b/Engine/Plugins/Developer/ProjectLauncher/Source/ProjectLauncher/Private/ProfileTree/CustomProfileTreeBuilder.cpp
@@ -16,14 +16,14 @@ namespace ProjectLauncher
 	{
 		FGenericProfileTreeBuilder::Construct();
 
-		FLaunchProfileTreeNode& GeneralSettingsHeader = TreeData->AddHeading( LOCTEXT("GeneralSettingsHeading", "General Settings") );
+		FLaunchProfileTreeNode& GeneralSettingsHeader = TreeData->AddHeading( TEXT("GeneralSettings"), LOCTEXT("GeneralSettingsHeading", "General Settings") );
 		AddProjectProperty(GeneralSettingsHeader);
 		AddTargetProperty(GeneralSettingsHeader);
 		AddPlatformProperty(GeneralSettingsHeader);
 		AddConfigurationProperty(GeneralSettingsHeader);
 		AddContentSchemeProperty(GeneralSettingsHeader);
 
-		FLaunchProfileTreeNode& ContentSchemeHeader = TreeData->AddHeading( LOCTEXT("ContentSchemeHeading", "Content Scheme") );
+		FLaunchProfileTreeNode& ContentSchemeHeader = TreeData->AddHeading( TEXT("ContentScheme"), LOCTEXT("ContentSchemeHeading", "Content Scheme") );
 		AddCompressPakFilesProperty(ContentSchemeHeader);
 		AddUseIoStoreProperty(ContentSchemeHeader);
 		AddGenerateChunksProperty(ContentSchemeHeader);
@@ -31,22 +31,22 @@ namespace ProjectLauncher
 		AddImportZenSnapshotProperty(ContentSchemeHeader);
 		AddZenPakStreamingPathProperty(ContentSchemeHeader);
 
-		FLaunchProfileTreeNode& MapsAndCookingHeader = TreeData->AddHeading( LOCTEXT("CookingHeading", "Maps And Cooking") );
+		FLaunchProfileTreeNode& MapsAndCookingHeader = TreeData->AddHeading( TEXT("Cooking"), LOCTEXT("CookingHeading", "Maps And Cooking") );
 		AddCookProperty(MapsAndCookingHeader);
 		AddIncrementalCookProperty(MapsAndCookingHeader);
 		AddMapsToCookProperty(MapsAndCookingHeader);
 		AddAdditionalCookerOptionsProperty(MapsAndCookingHeader);
 
-		FLaunchProfileTreeNode& BuildHeader = TreeData->AddHeading( LOCTEXT("BuildHeading", "Build") );
+		FLaunchProfileTreeNode& BuildHeader = TreeData->AddHeading( TEXT("Build"), LOCTEXT("BuildHeading", "Build") );
 		AddBuildProperty(BuildHeader);
 		AddForceBuildProperty(BuildHeader);
 		AddArchitectureProperty(BuildHeader);
 
-		FLaunchProfileTreeNode& DirectoryHeader = TreeData->AddHeading( LOCTEXT("DirectoryHeading", "Directory") );
+		FLaunchProfileTreeNode& DirectoryHeader = TreeData->AddHeading( TEXT("Directory"), LOCTEXT("DirectoryHeading", "Directory") );
 		AddArchiveBuildProperty(DirectoryHeader);
 		AddArchiveBuildDirectoryProperty(DirectoryHeader);
 
-		FLaunchProfileTreeNode& DeployAndRunHeader = TreeData->AddHeading( LOCTEXT("DeployAndRunHeading", "Deploy And Run") );
+		FLaunchProfileTreeNode& DeployAndRunHeader = TreeData->AddHeading( TEXT("DeployAndRun"), LOCTEXT("DeployAndRunHeading", "Deploy And Run") );
 		AddDeployProperty(DeployAndRunHeader);
 		AddTargetDeviceProperty(DeployAndRunHeader);
 		AddRunProperty(DeployAndRunHeader);
diff --git a/Engine/Plugins/Developer/ProjectLauncher/Source/ProjectLauncher/Private/ProfileTree/LaunchProfileTreeData.cpp b/Engine/Plugins/Developer/ProjectLauncher/Source/ProjectLauncher/Private/ProfileTree/LaunchProfileTreeData.cpp
--- a/Engine/Plugins/Developer/ProjectLauncher/Source/ProjectLauncher/Private/ProfileTree/LaunchProfileTreeData.cpp
+++ b/Engine/Plugins/Developer/ProjectLauncher/Source/ProjectLauncher/Private/ProfileTree/LaunchProfileTreeData.cpp
@@ -145,12 +145,25 @@ namespace ProjectLauncher
 	}
 
 
-	FLaunchProfileTreeNode& FLaunchProfileTreeData::AddHeading( FText InName )
+	FLaunchProfileTreeNode& FLaunchProfileTreeData::AddHeading( const TCHAR* InName, FText InDisplayName )
 	{
+		if (HeadingNodes.Contains(InName))
+		{
+			return GetHeading(InName);
+		}
+
 		FLaunchProfileTreeNodePtr TreeNode = MakeShared<FLaunchProfileTreeNode>(this);
-		TreeNode->Name = InName;
+		TreeNode->Name = InDisplayName;
 		Nodes.Add(TreeNode);
 
+		HeadingNodes.Add(InName, TreeNode);
+
+		return *TreeNode.Get();
+	}
+
+	FLaunchProfileTreeNode& FLaunchProfileTreeData::GetHeading( const TCHAR* InName ) const
+	{
+		FLaunchProfileTreeNodePtr TreeNode = HeadingNodes.FindChecked(InName);
 		return *TreeNode.Get();
 	}
 
diff --git a/Engine/Plugins/Developer/ProjectLauncher/Source/ProjectLauncher/Public/Extension/LaunchExtension.h b/Engine/Plugins/Developer/ProjectLauncher/Source/ProjectLauncher/Public/Extension/LaunchExtension.h
--- a/Engine/Plugins/Developer/ProjectLauncher/Source/ProjectLauncher/Public/Extension/LaunchExtension.h
+++ b/Engine/Plugins/Developer/ProjectLauncher/Source/ProjectLauncher/Public/Extension/LaunchExtension.h
@@ -101,6 +101,17 @@ namespace ProjectLauncher
 		 */
 		UE_API virtual FText GetCustomExtensionMenuSection() const;
 
+		/**
+		 * Create the default heading for this extension
+		 */
+		UE_API FLaunchProfileTreeNode& AddDefaultHeading(ProjectLauncher::FLaunchProfileTreeData& ProfileTreeData) const;
+
+
+		/**
+		 * Create the default heading for this extension, with a bespoke display name
+		 */
+		UE_API FLaunchProfileTreeNode& AddDefaultHeading(ProjectLauncher::FLaunchProfileTreeData& ProfileTreeData, FText OverrideDisplayName) const;
+
 		/*
 		 * Populate the custom extension menu for this extension. This menu would typically contain
 		 * items that might enable bespoke tree customization options for example.
diff --git a/Engine/Plugins/Developer/ProjectLauncher/Source/ProjectLauncher/Public/ProfileTree/LaunchProfileTreeData.h b/Engine/Plugins/Developer/ProjectLauncher/Source/ProjectLauncher/Public/ProfileTree/LaunchProfileTreeData.h
--- a/Engine/Plugins/Developer/ProjectLauncher/Source/ProjectLauncher/Public/ProfileTree/LaunchProfileTreeData.h
+++ b/Engine/Plugins/Developer/ProjectLauncher/Source/ProjectLauncher/Public/ProfileTree/LaunchProfileTreeData.h
@@ -150,7 +150,8 @@ namespace ProjectLauncher
 		UE_API ~FLaunchProfileTreeData();
 
 		// node creation
-		UE_API FLaunchProfileTreeNode& AddHeading( FText InName );
+		UE_API FLaunchProfileTreeNode& AddHeading( const TCHAR* InName, FText InDisplayName );
+		UE_API FLaunchProfileTreeNode& GetHeading( const TCHAR* InName ) const;
 		UE_API void CreateExtensionsUI();
 
 		UE_API void OnPropertyChanged();
@@ -160,6 +161,7 @@ namespace ProjectLauncher
 		ILauncherProfilePtr Profile;
 		TSharedRef<FModel> Model;
 		TArray<FLaunchProfileTreeNodePtr> Nodes;
+		TMap<FString,FLaunchProfileTreeNodePtr> HeadingNodes;
 		ILaunchProfileTreeBuilder* TreeBuilder;
 		TArray<TSharedPtr<FLaunchExtensionInstance>> ExtensionInstances;
 		bool bHasAnyMenuExtensions = false;
