Add comments for Project Launcher API functions

diff --git a/Engine/Plugins/Developer/ProjectLauncher/Source/ProjectLauncher/Public/Model/ProjectLauncherModel.h b/Engine/Plugins/Developer/ProjectLauncher/Source/ProjectLauncher/Public/Model/ProjectLauncherModel.h
--- a/Engine/Plugins/Developer/ProjectLauncher/Source/ProjectLauncher/Public/Model/ProjectLauncherModel.h
+++ b/Engine/Plugins/Developer/ProjectLauncher/Source/ProjectLauncher/Public/Model/ProjectLauncherModel.h
@@ -25,6 +25,7 @@ struct FTargetInfo;
 
 namespace ProjectLauncher
 {
+	/* Defines how content is stored and deployed for a profile */
 	enum class EContentScheme : uint8
 	{
 		PakFiles,
@@ -35,13 +36,26 @@ namespace ProjectLauncher
 		CookOnTheFly,
 		MAX
 	};
+
+	/* Returns all possible content schemes */
 	TArray<EContentScheme> GetAllContentSchemes();
+
+	/* Get user-visible display name for a content scheme */
 	FText GetContentSchemeDisplayName(EContentScheme ContentScheme);
+
+	/* Get tooltip text describing a content scheme */
 	FText GetContentSchemeToolTip(EContentScheme ContentScheme);
+
+	/* Returns a launch error message for a given profile, optionally with annotations */
 	FText GetProfileLaunchErrorMessage(ILauncherProfilePtr Profile, bool bWithAnnotations = false );
+
+	/* Convert a content scheme to string */
 	const TCHAR* LexToString( const ProjectLauncher::EContentScheme& ContentScheme);
+
+	/* Try to parse a string into a content scheme enum value */
 	bool LexTryParseString( ProjectLauncher::EContentScheme& OutContentScheme, const TCHAR* String );
 
+	/* Defines the type of a profile (basic, custom, or advanced) */
 	enum class EProfileType : uint8
 	{
 		Invalid,
@@ -50,6 +64,7 @@ namespace ProjectLauncher
 		Advanced,
 	};
 
+	/* Holds a single log message emitted during launch */
 	struct FLaunchLogMessage
 	{
 		TSharedRef<FString> Message;
@@ -61,6 +76,7 @@ namespace ProjectLauncher
 		{}
 	};
 
+	/* Cache of project settings relevant for launch profiles */
 	struct FProjectSettings
 	{
 		// whether this project is the one that is currently open in the editor
@@ -81,132 +97,195 @@ namespace ProjectLauncher
 		// default build target name
 		FString DefaultBuildTargetName;
 
+		// cached project config ini
 		FConfigCacheIni* Config = nullptr;
 	};
 
 
+	/* Delegate fired when a profile is clicked in the UI */
 	DECLARE_DELEGATE_OneParam(FOnProfileClicked, const ILauncherProfilePtr&)
 
+	/* Delegate fired when the selected profile changes */
 	DECLARE_MULTICAST_DELEGATE_TwoParams(FOnSelectedProfileChanged, const ILauncherProfilePtr& /*new*/, const ILauncherProfilePtr& /*old*/)
+
+	/* Delegate fired when the selected project changes */
 	DECLARE_MULTICAST_DELEGATE(FOnSelectedProjectChanged)
+
+	/* Delegate fired when the active build target changes */
 	DECLARE_MULTICAST_DELEGATE(FOnBuildTargetChanged)
+
+	/* Delegate fired when the selected platform changes */
 	DECLARE_MULTICAST_DELEGATE(FOnSelectedPlatformChanged)
 
+
 	extern bool bUseFriendlyBuildTargetSelection;
 
+	/*
+	 * This class contains a selection of helper functions and the current profile state
+	 */
 	class FModel : public TSharedFromThis<FModel>
 	{
 	public:
+		/* Constructor/destructor */
 		UE_API FModel(const TSharedRef<ITargetDeviceProxyManager>& InDeviceProxyManager, const TSharedRef<ILauncher>& InLauncher, const TSharedRef<ILauncherProfileManager>& InProfileManager);
 		UE_API ~FModel();
 
-
+		/* Get the target device proxy manager */
 		const TSharedRef<ITargetDeviceProxyManager>& GetDeviceProxyManager() const
 		{
 			return DeviceProxyManager;
 		}
 
+		/* Get the launcher instance */
 		const TSharedRef<ILauncher>& GetLauncher() const
 		{
 			return Launcher;
 		}
 
+		/* Get the launcher profile manager */
 		const TSharedRef<ILauncherProfileManager>& GetProfileManager() const
 		{
 			return ProfileManager;
 		}
 
-
+		/* Get all known launcher profiles */
 		const TArray<ILauncherProfilePtr>& GetAllProfiles() const
 		{
 			return AllProfiles;
 		}
 
+		/* Get the basic launch profile */
 		ILauncherProfilePtr GetBasicLaunchProfile() const
 		{
 			return BasicLaunchProfile;
 		}
 
+		/* Get the currently selected profile */
 		const ILauncherProfilePtr& GetSelectedProfile() const
 		{
 			return SelectedProfile;
 		}
 
+		/* Delegate fired when a profile is selected */
 		FOnSelectedProfileChanged& OnProfileSelected()
 		{
 			return ProfileSelectedDelegate;
 		}
 
+		/* Change the selected profile */
 		UE_API void SelectProfile(const ILauncherProfilePtr& Profile);
+
+		/* Sort all profiles alphabetically, keeping the basic profile first in the list */
 		UE_API void SortProfiles();
 
+		/* Get the type of a profile (basic, custom, advanced) */
 		UE_API EProfileType GetProfileType(const ILauncherProfileRef& Profile) const;
+
+		/* Returns true if the given profile is advanced (multiple platforms/cultures, DLC, patches, etc) */
 		UE_API bool IsAdvancedProfile(const ILauncherProfileRef& Profile) const;
-		UE_API bool IsBasicLaunchProfile( const ILauncherProfilePtr& Profile ) const;
 
-		UE_API ILauncherProfileRef CreateCustomProfile( const TCHAR* Name );
+		/* Returns true if the given profile is the basic launch profile */
+		UE_API bool IsBasicLaunchProfile(const ILauncherProfilePtr& Profile) const;
+
+		/* Create a new custom profile with the given name */
+		UE_API ILauncherProfileRef CreateCustomProfile(const TCHAR* Name);
+
+		/* Create a new basic launch profile */
 		UE_API ILauncherProfileRef CreateBasicLaunchProfile();
+
+		/* Clone an existing custom profile */
 		UE_API ILauncherProfilePtr CloneCustomProfile(const ILauncherProfileRef& Profile);
 
+		/* Get the default basic launch profile - used for "reset to default"  */
 		UE_API const ILauncherProfileRef GetDefaultBasicLaunchProfile() const;
+
+		/* Get the default custom launch profile - used for "reset to default" */
 		UE_API const ILauncherProfileRef GetDefaultCustomLaunchProfile() const;
 
+		/* Determine the content scheme for a profile */
 		UE_API EContentScheme DetermineProfileContentScheme(const ILauncherProfileRef& Profile) const;
+
+		/* Apply a content scheme to a profile, setting cook/deploy/packaging options etc */
 		UE_API void SetProfileContentScheme(EContentScheme ContentScheme, const ILauncherProfileRef& Profile, bool bWantToCook = true, ELauncherProfileDeploymentModes::Type DefaultDeploymentMode = ELauncherProfileDeploymentModes::CopyToDevice);
 
-		UE_API TSharedPtr<FLaunchLogMessage> AddLogMessage( const FString& InMessage, ELogVerbosity::Type InVerbosity = ELogVerbosity::Log );
-		UE_API void ClearLogMessages();
-		int32 GetNumLogMessages() const { return LaunchLogMessages.Num(); }
+		/* Record a new log message from launching */
+		UE_API TSharedPtr<FLaunchLogMessage> AddLogMessage(const FString& InMessage, ELogVerbosity::Type InVerbosity = ELogVerbosity::Log);
 
+		/* Clear all log messages */
+		UE_API void ClearLogMessages();
+		
+		/* Get number of log messages */
+		int32 GetNumLogMessages() const
+		{
+			return LaunchLogMessages.Num();
+		}
 
+		/* Get platform info for a profile or platform name */
 		static UE_API const PlatformInfo::FTargetPlatformInfo* GetPlatformInfo(const ILauncherProfilePtr& Profile);
 		static UE_API const PlatformInfo::FTargetPlatformInfo* GetPlatformInfo(FName PlatformName, const FTargetInfo& BuildTargetInfo);
 
+		/* Returns true if profile or platform name refers to host platform */
 		static UE_API bool IsHostPlatform(const ILauncherProfilePtr& Profile);
 		static UE_API bool IsHostPlatform(FName PlatformName);
 
-		UE_API FTargetInfo GetBuildTargetInfo( FString BuildTargetName, const FString& ProjectPath );
-		UE_API FTargetInfo GetBuildTargetInfo( const ILauncherProfileRef& Profile );
+		/* Get build target info for a project path or profile */
+		UE_API FTargetInfo GetBuildTargetInfo(FString BuildTargetName, const FString& ProjectPath);
+		UE_API FTargetInfo GetBuildTargetInfo(const ILauncherProfileRef& Profile);
 
+		/* Get a device proxy for a profile */
 		static UE_API TSharedPtr<ITargetDeviceProxy> GetDeviceProxy(const ILauncherProfileRef& Profile);
+
+		/* Update cooked platforms based on selected device proxy */
 		UE_API void UpdatedCookedPlatformsFromDeployDeviceProxy(const ILauncherProfileRef& Profile, TSharedPtr<ITargetDeviceProxy> DeviceProxy = nullptr);
+
+		/* Update cooked platforms based on build target */
 		UE_API void UpdateCookedPlatformsFromBuildTarget(const ILauncherProfileRef& Profile);
 
-		static UE_API FString GetVanillaPlatformName( const FString& PlatformName );
-		static UE_API FString GetBuildTargetPlatformName( const FString& PlatformName, const FTargetInfo& BuildTargetInfo );
+		/* Utility functions for resolving platform names */
+		static UE_API FString GetVanillaPlatformName(const FString& PlatformName);
+		static UE_API FString GetBuildTargetPlatformName(const FString& PlatformName, const FTargetInfo& BuildTargetInfo);
 
-		UE_API const FProjectSettings GetProjectSettings( const FString& ProjectPath );
-		UE_API const FProjectSettings GetProjectSettings( const ILauncherProfileRef& Profile );
+		/* Retrieve project settings from a project path or profile */
+		UE_API const FProjectSettings GetProjectSettings(const FString& ProjectPath);
+		UE_API const FProjectSettings GetProjectSettings(const ILauncherProfileRef& Profile);
 
+		/* Get config section and ini file used for storing project launcher state */
 		UE_API const TCHAR* GetConfigSection() const;
 		UE_API const FString& GetConfigIni() const;
 
+		/* Returns true if extensions are enabled for project launcher */
 		UE_API bool AreExtensionsEnabled() const;
 
-		UE_API TArray<FString> GetAvailableProjectMapNames( const FString& InProjectPath, bool bIncludeNonContentDirMaps = false);
-		UE_API TArray<FString> GetAvailableProjectMapPaths( const FString& InProjectPath, bool bIncludeNonContentDirMaps = false);
+		/* Get available maps for a project or engine */
+		UE_API TArray<FString> GetAvailableProjectMapNames(const FString& InProjectPath, bool bIncludeNonContentDirMaps = false);
+		UE_API TArray<FString> GetAvailableProjectMapPaths(const FString& InProjectPath, bool bIncludeNonContentDirMaps = false);
 		UE_API TArray<FString> GetAvailableEngineMapNames();
 		UE_API TArray<FString> GetAvailableEngineMapPaths();
 
 	private:
-		FString ConfigFileName;
-
+		/* Load and save config state */
 		UE_API void LoadConfig();
 		UE_API void SaveConfig();
 
+		/* Handlers for profile manager and device proxy events */
 		UE_API void HandleProfileManagerProfileAdded(const ILauncherProfileRef& Profile);
 		UE_API void HandleProfileManagerProfileRemoved(const ILauncherProfileRef& Profile);
-		UE_API void HandleDeviceProxyAdded(const TSharedRef<ITargetDeviceProxy>& DeviceProxy );
-		UE_API void HandleDeviceProxyRemoved(const TSharedRef<ITargetDeviceProxy>& DeviceProxy );
+		UE_API void HandleDeviceProxyAdded(const TSharedRef<ITargetDeviceProxy>& DeviceProxy);
+		UE_API void HandleDeviceProxyRemoved(const TSharedRef<ITargetDeviceProxy>& DeviceProxy);
 
-		UE_API void ReadProjectSettingsFromConfig( FConfigCacheIni& InConfig, const FString& InProjectPath, FProjectSettings& OutResult );
-
-		UE_API void OnModifyLaunchCommandLine( const ILauncherProfileRef& InProfile, FString& InOutCommandLine );
-		UE_API void OnModifyAutomatedTestCommandLine( const ILauncherProfileAutomatedTestRef& InAutomatedTest, const ILauncherProfileRef& InProfile, FString& InOutCommandLine );
+		/* Helper for reading project settings from config */
+		UE_API void ReadProjectSettingsFromConfig(FConfigCacheIni& InConfig, const FString& InProjectPath, FProjectSettings& OutResult);
 
+		/* Modify launch command lines to inject extension variables */
+		UE_API void OnModifyLaunchCommandLine(const ILauncherProfileRef& InProfile, FString& InOutCommandLine);
+		UE_API void OnModifyAutomatedTestCommandLine(const ILauncherProfileAutomatedTestRef& InAutomatedTest, const ILauncherProfileRef& InProfile, FString& InOutCommandLine);
 
+		/* Get map paths for a project, caching results */
 		UE_API TArray<FString> GetAndCacheMapPaths(const FString& InOptionalProjectPath, bool bIncludeNonContentDirMaps = false);
 
+	private:
+		FString ConfigFileName;
+		
 		TSharedRef<ITargetDeviceProxyManager> DeviceProxyManager;
 		TSharedRef<ILauncher> Launcher;
 		TSharedRef<ILauncherProfileManager> ProfileManager;
diff --git a/Engine/Plugins/Developer/ProjectLauncher/Source/ProjectLauncher/Public/ProfileTree/LaunchProfileTreeData.h b/Engine/Plugins/Developer/ProjectLauncher/Source/ProjectLauncher/Public/ProfileTree/LaunchProfileTreeData.h
--- a/Engine/Plugins/Developer/ProjectLauncher/Source/ProjectLauncher/Public/ProfileTree/LaunchProfileTreeData.h
+++ b/Engine/Plugins/Developer/ProjectLauncher/Source/ProjectLauncher/Public/ProfileTree/LaunchProfileTreeData.h
@@ -15,6 +15,7 @@
 
 class SWidget;
 struct EVisibility;
+
 namespace ProjectLauncher
 {
 	class ILaunchProfileTreeBuilder;
@@ -29,6 +30,17 @@ namespace ProjectLauncher
 	typedef TSharedPtr<FLaunchProfileTreeNode> FLaunchProfileTreeNodePtr;
 	typedef TSharedRef<FLaunchProfileTreeNode> FLaunchProfileTreeNodeRef;
 
+	/*
+	 * Validation object representing errors and warnings tied to profile settings. 
+	 * If the given validation errors are present in the profile a warning icon will be displayed next to this property in the UI
+	 *
+	 * Use with the Validation callback in FLaunchProfileTreeNode:
+	 *
+	 *  .Validation = FValidation( {ELauncherProfileValidationErrors::NoProjectSelected, ...} ),
+	 *     or
+	 *  .Validation = FValidation( {TEXT("Validation_SomeCustomValidation", ...} ),
+	 * 
+	 */
 	class FValidation
 	{
 	public:
@@ -37,23 +49,29 @@ namespace ProjectLauncher
 		UE_API FValidation( std::initializer_list<ELauncherProfileValidationErrors::Type> InAssociatedErrors, std::initializer_list<FString> InAssociatedCustomErrors );
 		UE_API FValidation( std::initializer_list<FString> InAssociatedCustomErrors );
 
+		/* Returns true if any validation rules are set. */
 		bool IsSet() const;
 
+		/* Returns true if the given profile has one of the associated errors. */
 		bool HasError( ILauncherProfileRef Profile ) const;
+
+		/* Collects error text from the profile for display. */
 		void GetErrorText( ILauncherProfileRef Profile, TArray<FString>& ErrorLines ) const;
 
 	private:
-		bool bIsSet = false;
-		TArray<ELauncherProfileValidationErrors::Type> AssociatedErrors;
-		TArray<FString> AssociatedCustomErrors;
+		bool bIsSet = false;                                             // whether this validation has been configured
+		TArray<ELauncherProfileValidationErrors::Type> AssociatedErrors; // engine-defined errors
+		TArray<FString> AssociatedCustomErrors;                          // custom errors specific to this profile
 	};
 
-
+	/*
+	 * Represents a single editable node (field, checkbox, widget, etc.) in the profile tree UI.
+	 */
 	class FLaunchProfileTreeNode : public TSharedFromThis<FLaunchProfileTreeNode>
 	{
 	public:
-		// construction
-		UE_API FLaunchProfileTreeNode( FLaunchProfileTreeData* InTreeData );
+		/* Construct a node bound to a tree. */
+		UE_API FLaunchProfileTreeNode(FLaunchProfileTreeData* InTreeData);
 
 		typedef TFunction<void()> FFunction;
 		typedef TFunction<bool()> FGetBool;
@@ -65,19 +83,23 @@ namespace ProjectLauncher
 		typedef TFunction<FString(void)> FGetString;
 		typedef TFunction<void(FString)> FSetString;
 
-		// node creation
+		/* Callbacks common to all widgets. */
 		struct FCallbacks
 		{
-			FGetBool IsDefault = nullptr;
-			FFunction SetToDefault = nullptr;
-			FGetBool IsVisible = nullptr;
-			FGetBool IsEnabled = nullptr;
-			FValidation Validation;
+			FGetBool IsDefault = nullptr;     // check if current value is default
+			FFunction SetToDefault = nullptr; // reset to default
+			FGetBool IsVisible = nullptr;     // visibility logic
+			FGetBool IsEnabled = nullptr;     // enabled/disabled logic
+			FValidation Validation;           // validation rules
 		};
+
+		/* Add a generic widget with callbacks. */
 		UE_API FLaunchProfileTreeNode& AddWidget( FText InName, FCallbacks&& InWidgetCallbacks, TSharedRef<SWidget> InValueWidget );
-		UE_API FLaunchProfileTreeNode& AddWidget( FText InName, TSharedRef<SWidget> InValueWidget );
 
+		/* Add a generic widget without callbacks. */
+		UE_API FLaunchProfileTreeNode& AddWidget( FText InName, TSharedRef<SWidget> InValueWidget );
 
+		/* Boolean-specific callbacks. */
 		struct FBooleanCallbacks
 		{
 			FGetBool GetValue = nullptr;
@@ -89,7 +111,7 @@ namespace ProjectLauncher
 		};
 		UE_API FLaunchProfileTreeNode& AddBoolean( FText InName, FBooleanCallbacks&& BooleanCallbacks );
 
-
+		/* String-specific callbacks. */
 		struct FStringCallbacks
 		{
 			FGetString GetValue = nullptr;
@@ -104,7 +126,7 @@ namespace ProjectLauncher
 		UE_API FLaunchProfileTreeNode& AddFileString(FText InName, FStringCallbacks&& StringCallbacks, FString FileFilter = TEXT("All files (*.*)|*.*"));
 		UE_API FLaunchProfileTreeNode& AddCommandLineString( FText InName, FStringCallbacks&& StringCallbacks );
 
-
+		/* Generic value callbacks. */
 		template<typename T>
 		struct TTypeCallbacks
 		{
@@ -115,57 +137,59 @@ namespace ProjectLauncher
 			FGetBool IsEnabled = nullptr;
 			FValidation Validation;
 		};
+
+		/* Add a typed value widget, e.g. spinner for int or float. */
 		template<typename T>
 		FLaunchProfileTreeNode& AddValue( FText InName, TTypeCallbacks<T>&& TypeCallbacks, TInterval<T> Range = TInterval<T>(), bool bShowSpinner = false );
 
-
 		UE_API FLaunchProfileTreeNode& AddInteger( FText InName, TTypeCallbacks<int32>&& TypeCallbacks, FInt32Interval Range = FInt32Interval(), bool bShowSpinner = false );
 		UE_API FLaunchProfileTreeNode& AddFloat( FText InName, TTypeCallbacks<float>&& TypeCallbacks, FFloatInterval Range = FFloatInterval(), bool bShowSpinner = false );
 
-
-
-		// member variables
-		FText Name;
-		TSharedPtr<SWidget> Widget;
-
-		FCallbacks Callbacks;
-
-
-		TArray<FLaunchProfileTreeNodePtr> Children;
+		/* Node details */
+		FText Name;                                 // label for the node
+		TSharedPtr<SWidget> Widget;                 // UI widget
+		FCallbacks Callbacks;                       // generic callbacks
+		TArray<FLaunchProfileTreeNodePtr> Children; // child nodes
 		const FLaunchProfileTreeData* GetTreeData() const { return TreeData; }
 
 	protected:
+		/* Notify parent tree that a property has changed. */
 		UE_API void OnPropertyChanged();
 
 		FLaunchProfileTreeData* TreeData;
 	};
 
-
-
+	/*
+	 * Represents the full profile tree of editable UI nodes for a launcher profile.
+	 * Integrates extensions and manages node headings.
+	 */
 	class FLaunchProfileTreeData : public TSharedFromThis<FLaunchProfileTreeData>
 	{
 	public:
-		// construction/destruction
+		/* Construct with a profile, model, and tree builder. */
 		UE_API FLaunchProfileTreeData(ILauncherProfilePtr InProfile, TSharedRef<FModel> InModel, ILaunchProfileTreeBuilder* InTreeBuilder);
 		UE_API ~FLaunchProfileTreeData();
 
-		// node creation
-		UE_API FLaunchProfileTreeNode& AddHeading( const TCHAR* InName, FText InDisplayName );
-		UE_API FLaunchProfileTreeNode& GetHeading( const TCHAR* InName ) const;
+		/* Add or retrieve top-level headings. */
+		UE_API FLaunchProfileTreeNode& AddHeading(const TCHAR* InName, FText InDisplayName);
+		UE_API FLaunchProfileTreeNode& GetHeading(const TCHAR* InName) const;
+
+		/* Invoke extension instances to build their UI. */
 		UE_API void CreateExtensionsUI();
 
+		/* Notify of property changes and refresh requests. */
 		UE_API void OnPropertyChanged();
 		UE_API void RequestTreeRefresh();
 
-		// member variables
-		ILauncherProfilePtr Profile;
-		TSharedRef<FModel> Model;
-		TArray<FLaunchProfileTreeNodePtr> Nodes;
-		TMap<FString,FLaunchProfileTreeNodePtr> HeadingNodes;
-		ILaunchProfileTreeBuilder* TreeBuilder;
-		TArray<TSharedPtr<FLaunchExtensionInstance>> ExtensionInstances;
-		bool bHasAnyMenuExtensions = false;
-		bool bRequestTreeRefresh = false;
+		/* Member variables */
+		ILauncherProfilePtr Profile;                                     // associated profile
+		TSharedRef<FModel> Model;                                        // owning model
+		TArray<FLaunchProfileTreeNodePtr> Nodes;                         // all nodes
+		TMap<FString,FLaunchProfileTreeNodePtr> HeadingNodes;            // named top-level headings
+		ILaunchProfileTreeBuilder* TreeBuilder;                          // UI builder interface
+		TArray<TSharedPtr<FLaunchExtensionInstance>> ExtensionInstances; // extension instances attached
+		bool bHasAnyMenuExtensions = false;                              // whether any extensions add menus
+		bool bRequestTreeRefresh = false;                                // request to refresh tree UI
 	};
 
 
