fix comment

#rnx

diff --git a/Engine/Plugins/Developer/ProjectLauncher/Source/ProjectLauncher/Public/ProfileTree/LaunchProfileTreeData.h b/Engine/Plugins/Developer/ProjectLauncher/Source/ProjectLauncher/Public/ProfileTree/LaunchProfileTreeData.h
--- a/Engine/Plugins/Developer/ProjectLauncher/Source/ProjectLauncher/Public/ProfileTree/LaunchProfileTreeData.h
+++ b/Engine/Plugins/Developer/ProjectLauncher/Source/ProjectLauncher/Public/ProfileTree/LaunchProfileTreeData.h
@@ -38,7 +38,7 @@ namespace ProjectLauncher
 	 *
 	 *  .Validation = FValidation( {ELauncherProfileValidationErrors::NoProjectSelected, ...} ),
 	 *     or
-	 *  .Validation = FValidation( {TEXT("Validation_SomeCustomValidation", ...} ),
+	 *  .Validation = FValidation( {TEXT("Validation_SomeCustomValidation"), ...} ),
 	 * 
 	 */
 	class FValidation
