Project Launcher extensions menu improvements:
- extensions can decide if they appear in their own menu section, a common section or a submenu.
- sections and submenus are merged if multiple extensions target them.
- automated tests default to being in a dedicated shared submenu, and all other extensions default to their own section.

#rnx

diff --git a/Engine/Plugins/Developer/ProjectLauncher/Source/ProjectLauncher/Private/Extension/LaunchExtension.cpp b/Engine/Plugins/Developer/ProjectLauncher/Source/ProjectLauncher/Private/Extension/LaunchExtension.cpp
--- a/Engine/Plugins/Developer/ProjectLauncher/Source/ProjectLauncher/Private/Extension/LaunchExtension.cpp
+++ b/Engine/Plugins/Developer/ProjectLauncher/Source/ProjectLauncher/Private/Extension/LaunchExtension.cpp
@@ -634,11 +634,6 @@ namespace ProjectLauncher
 		PropertyChangedDelegate.Broadcast();
 	}
 
-	FText FLaunchExtensionInstance::GetCustomExtensionMenuSection() const
-	{
-		return Extension->GetDisplayName();
-	}
-
 	FLaunchProfileTreeNode& FLaunchExtensionInstance::AddDefaultHeading(ProjectLauncher::FLaunchProfileTreeData& ProfileTreeData) const
 	{
 		return AddDefaultHeading(ProfileTreeData, Extension->GetDisplayName());
@@ -649,6 +644,18 @@ namespace ProjectLauncher
 		return ProfileTreeData.AddHeading(Extension->GetInternalName(), OverrideDisplayName);
 	}
 
+	void FLaunchExtensionInstance::GetExtensionsMenuEntry( FExtensionsMenuEntry& MenuEntry ) const
+	{
+		// put the extension in its own section by default
+		MenuEntry.Type = FExtensionsMenuEntry::Type_Section;
+		MenuEntry.SectionName = Extension->GetInternalName();
+		MenuEntry.SectionDisplayName = Extension->GetDisplayName();
+		
+		// fill in the default submenu details anyway - an extension can then just change MenuEntry.Type = Type_SubMenu to appear in a named submenu
+		MenuEntry.SubmenuName = Extension->GetInternalName();
+		MenuEntry.SubmenuDisplayName = Extension->GetDisplayName();
+	}
+
 
 
 
@@ -703,16 +710,19 @@ namespace ProjectLauncher
 		return GetProfile()->GetAutomatedTest(*GetTestInternalName());
 	}
 
-	FText FAutomatedTestLaunchExtensionInstance::GetDefaultExtensionMenuSection()
+	void FAutomatedTestLaunchExtensionInstance::GetExtensionsMenuEntry( FExtensionsMenuEntry& MenuEntry ) const
 	{
-		return LOCTEXT("AutotestMenuSection", "Automated Tests");
-	}
+		// put all automated tests in shared submenu by default
+		MenuEntry.Type = FExtensionsMenuEntry::Type_SubMenu;
+		MenuEntry.SubmenuName = TEXT("AutomatedTestsSection");
+		MenuEntry.SubmenuDisplayName = LOCTEXT("AutotestMenuSection", "Automated Tests");
+
+		// put the submenu into the 'common' category
+		MenuEntry.SectionName = FString();
 
-	FText FAutomatedTestLaunchExtensionInstance::GetCustomExtensionMenuSection() const
-	{
-		return GetDefaultExtensionMenuSection();
 	}
 
+
 	void FAutomatedTestLaunchExtensionInstance::CustomizeAutomatedTestCommandLine( const ILauncherProfileAutomatedTestRef& InAutomatedTest, FString& InOutCommandLine )
 	{
 		if (InAutomatedTest == GetTest())
diff --git a/Engine/Plugins/Developer/ProjectLauncher/Source/ProjectLauncher/Private/ProfileTree/LaunchProfileTreeData.cpp b/Engine/Plugins/Developer/ProjectLauncher/Source/ProjectLauncher/Private/ProfileTree/LaunchProfileTreeData.cpp
--- a/Engine/Plugins/Developer/ProjectLauncher/Source/ProjectLauncher/Private/ProfileTree/LaunchProfileTreeData.cpp
+++ b/Engine/Plugins/Developer/ProjectLauncher/Source/ProjectLauncher/Private/ProfileTree/LaunchProfileTreeData.cpp
@@ -123,15 +123,9 @@ namespace ProjectLauncher
 		{
 			ExtensionInstances = FLaunchExtension::CreateExtensionInstancesForProfile(Profile.ToSharedRef(), InModel);
 
-			bHasAnyMenuExtensions = false;
 			for (TSharedPtr<FLaunchExtensionInstance> ExtensionInstance : ExtensionInstances)
 			{
 				ExtensionInstance->GetPropertyChangedDelegate().AddRaw(TreeBuilder, &ILaunchProfileTreeBuilder::OnPropertyChanged );
-
-				if (ExtensionInstance->HasCustomExtensionMenu())
-				{
-					bHasAnyMenuExtensions = true;
-				}
 			}
 		}
 	}
diff --git a/Engine/Plugins/Developer/ProjectLauncher/Source/ProjectLauncher/Private/Widgets/Profiles/SCustomLaunchCustomProfileEditor.cpp b/Engine/Plugins/Developer/ProjectLauncher/Source/ProjectLauncher/Private/Widgets/Profiles/SCustomLaunchCustomProfileEditor.cpp
--- a/Engine/Plugins/Developer/ProjectLauncher/Source/ProjectLauncher/Private/Widgets/Profiles/SCustomLaunchCustomProfileEditor.cpp
+++ b/Engine/Plugins/Developer/ProjectLauncher/Source/ProjectLauncher/Private/Widgets/Profiles/SCustomLaunchCustomProfileEditor.cpp
@@ -524,49 +524,120 @@ void SCustomLaunchCustomProfileEditor::SetProfile( const ILauncherProfilePtr& Pr
 	TreeView->RequestTreeRefresh();
 }
 
-bool SCustomLaunchCustomProfileEditor::HasAnyMenuExtensions() const
+TSharedRef<SWidget> SCustomLaunchCustomProfileEditor::MakeExtensionsMenu()
 {
-	if (TreeBuilder && TreeBuilder->GetProfileTree()->bHasAnyMenuExtensions)
+	ProjectLauncher::FLaunchProfileTreeDataRef TreeData = TreeBuilder->GetProfileTree();
+
+	// ...build extension submenu structure...
+
+	struct FExtensionMenuSubMenu
+	{
+		FText DisplayName;
+		TArray<TSharedPtr<ProjectLauncher::FLaunchExtensionInstance>> ExtensionInstances;
+	};
+
+	struct FExtensionMenuSection
 	{
-		return true;
+		FText DisplayName;
+		TMap<FString,TSharedPtr<FExtensionMenuSubMenu>> SubMenus;
+		TArray<TSharedPtr<ProjectLauncher::FLaunchExtensionInstance>> ExtensionInstances;
+	};
+
+	TMap<FString,FExtensionMenuSection> Sections;
+
+	for (TSharedPtr<ProjectLauncher::FLaunchExtensionInstance> ExtensionInstance : TreeData->ExtensionInstances)
+	{
+		ProjectLauncher::FLaunchExtensionInstance::FExtensionsMenuEntry ExtensionMenuEntry;
+		ExtensionInstance->GetExtensionsMenuEntry(ExtensionMenuEntry);
+		if (ExtensionMenuEntry.Type == ProjectLauncher::FLaunchExtensionInstance::FExtensionsMenuEntry::Type_None)
+		{
+			continue;
+		}
+
+		// get section
+		FExtensionMenuSection* Section = Sections.Find(ExtensionMenuEntry.SectionName);
+		if (Section == nullptr)
+		{
+			Section = &Sections.Add(ExtensionMenuEntry.SectionName);
+			Section->DisplayName = ExtensionMenuEntry.SectionName.IsEmpty() ? FText::GetEmpty() : ExtensionMenuEntry.SectionDisplayName;
+		}
+
+		if (ExtensionMenuEntry.Type == ProjectLauncher::FLaunchExtensionInstance::FExtensionsMenuEntry::Type_SubMenu)
+		{
+			// get submenu
+			TSharedPtr<FExtensionMenuSubMenu>* SubMenu = Section->SubMenus.Find(ExtensionMenuEntry.SubmenuName);
+			if (SubMenu == nullptr)
+			{
+				SubMenu = &Section->SubMenus.Add(ExtensionMenuEntry.SubmenuName, MakeShared<FExtensionMenuSubMenu>());
+				(*SubMenu)->DisplayName = ExtensionMenuEntry.SubmenuDisplayName;
+			}
+
+			(*SubMenu)->ExtensionInstances.Add(ExtensionInstance);
+		}
+		else if (ExtensionMenuEntry.Type == ProjectLauncher::FLaunchExtensionInstance::FExtensionsMenuEntry::Type_Section)
+		{
+			Section->ExtensionInstances.Add(ExtensionInstance);
+		}
 	}
 
-	return false;
-}
+	// sort sections alphabetically : this should leave the 'empty' section at the top
+	Sections.ValueSort([](const FExtensionMenuSection& A, const FExtensionMenuSection& B)
+	{
+		return A.DisplayName.CompareToCaseIgnored(B.DisplayName) < 0; 
+	});
+
 
-TSharedRef<SWidget> SCustomLaunchCustomProfileEditor::MakeExtensionsMenu()
-{
-	ProjectLauncher::FLaunchProfileTreeDataRef TreeData = TreeBuilder->GetProfileTree();
+
+	// ...construct the extensions menu...
 
 	const bool bShouldCloseWindowAfterMenuSelection = false;
 	const bool bCloseSelfOnly = false;
 	const bool bSearchable = false;
 	FMenuBuilder MenuBuilder(bShouldCloseWindowAfterMenuSelection, nullptr, nullptr, bCloseSelfOnly, &FCoreStyle::Get(), bSearchable );
-
-	for (TSharedPtr<ProjectLauncher::FLaunchExtensionInstance> ExtensionInstance : TreeData->ExtensionInstances)
+	
+	for (const auto& SectionItr : Sections)
 	{
-		if (!ExtensionInstance->HasCustomExtensionMenu())
+		// start the section if necessary
+		if (!SectionItr.Key.IsEmpty())
 		{
-			continue;
+			MenuBuilder.BeginSection(NAME_None, SectionItr.Value.DisplayName);
 		}
 
-#if 1
-		// flat list of extensions - easier to see but risks getting cluttered later on
-		MenuBuilder.BeginSection(NAME_None, ExtensionInstance->GetCustomExtensionMenuSection());
-		ExtensionInstance->MakeCustomExtensionSubmenu(MenuBuilder);
-		MenuBuilder.EndSection();
-#else 
-		// extensions in submenus - more scalable, but harder to see everything at a glance
-		const bool bInOpenSubMenuOnClick = false;
-		MenuBuilder.AddSubMenu(
-			ExtensionInstance->GetCustomExtensionMenuSection()),
-			FText::GetEmpty(),
-			FNewMenuDelegate::CreateSP( ExtensionInstance.ToSharedRef(), &FLaunchExtensionInstance::MakeExtensionSubmenu),
-			bInOpenSubMenuOnClick,
-			FSlateIcon(),
-			bShouldCloseWindowAfterMenuSelection
-		);
-#endif
+		// add all extensions
+		for (TSharedPtr<ProjectLauncher::FLaunchExtensionInstance> ExtensionInstance : SectionItr.Value.ExtensionInstances)
+		{
+			ExtensionInstance->MakeCustomExtensionSubmenu(MenuBuilder);
+		}
+
+		// add all submenus
+		for (const auto& SubMenuItr : SectionItr.Value.SubMenus)
+		{
+			auto MakeExtensionSubMenu = [](FMenuBuilder& MenuBuilder, TSharedPtr<FExtensionMenuSubMenu> SubMenu)
+			{
+				MenuBuilder.BeginSection(NAME_None, SubMenu->DisplayName);
+				for (TSharedPtr<ProjectLauncher::FLaunchExtensionInstance> ExtensionInstance : SubMenu->ExtensionInstances)
+				{
+					ExtensionInstance->MakeCustomExtensionSubmenu(MenuBuilder);
+				}
+				MenuBuilder.EndSection();
+			};
+
+			const bool bInOpenSubMenuOnClick = false;
+			MenuBuilder.AddSubMenu(
+				SubMenuItr.Value->DisplayName,
+				FText::GetEmpty(),
+				FNewMenuDelegate::CreateLambda( MakeExtensionSubMenu, SubMenuItr.Value),
+				bInOpenSubMenuOnClick,
+				FSlateIcon(),
+				bShouldCloseWindowAfterMenuSelection
+			);
+		}
+
+		// end the section, if necessary
+		if (!SectionItr.Key.IsEmpty())
+		{
+			MenuBuilder.EndSection();
+		}
 	}
 
 	return MenuBuilder.MakeWidget();
diff --git a/Engine/Plugins/Developer/ProjectLauncher/Source/ProjectLauncher/Private/Widgets/Profiles/SCustomLaunchCustomProfileEditor.h b/Engine/Plugins/Developer/ProjectLauncher/Source/ProjectLauncher/Private/Widgets/Profiles/SCustomLaunchCustomProfileEditor.h
--- a/Engine/Plugins/Developer/ProjectLauncher/Source/ProjectLauncher/Private/Widgets/Profiles/SCustomLaunchCustomProfileEditor.h
+++ b/Engine/Plugins/Developer/ProjectLauncher/Source/ProjectLauncher/Private/Widgets/Profiles/SCustomLaunchCustomProfileEditor.h
@@ -32,7 +32,6 @@ public:
 	void SetProfile( const ILauncherProfilePtr& Profile );
 	ILauncherProfilePtr GetProfile() const { return CurrentProfile; }
 
-	bool HasAnyMenuExtensions() const;
 	TSharedRef<SWidget> MakeExtensionsMenu();
 
 	// IScrollableWidget interface
diff --git a/Engine/Plugins/Developer/ProjectLauncher/Source/ProjectLauncher/Private/Widgets/Profiles/SCustomLaunchProfilesPanel.cpp b/Engine/Plugins/Developer/ProjectLauncher/Source/ProjectLauncher/Private/Widgets/Profiles/SCustomLaunchProfilesPanel.cpp
--- a/Engine/Plugins/Developer/ProjectLauncher/Source/ProjectLauncher/Private/Widgets/Profiles/SCustomLaunchProfilesPanel.cpp
+++ b/Engine/Plugins/Developer/ProjectLauncher/Source/ProjectLauncher/Private/Widgets/Profiles/SCustomLaunchProfilesPanel.cpp
@@ -564,7 +564,7 @@ END_SLATE_FUNCTION_BUILD_OPTIMIZATION
 
 EVisibility SCustomLaunchProfilesPanel::GetExtensionsBarVisibility() const
 {
-	if (IsProfileEditorVisible() && Model->AreExtensionsEnabled() && PropertyEditor->HasAnyMenuExtensions())
+	if (IsProfileEditorVisible() && Model->AreExtensionsEnabled())
 	{
 		return EVisibility::Visible;
 	}
diff --git a/Engine/Plugins/Developer/ProjectLauncher/Source/ProjectLauncher/Public/Extension/LaunchExtension.h b/Engine/Plugins/Developer/ProjectLauncher/Source/ProjectLauncher/Public/Extension/LaunchExtension.h
--- a/Engine/Plugins/Developer/ProjectLauncher/Source/ProjectLauncher/Public/Extension/LaunchExtension.h
+++ b/Engine/Plugins/Developer/ProjectLauncher/Source/ProjectLauncher/Public/Extension/LaunchExtension.h
@@ -92,21 +92,34 @@ namespace ProjectLauncher
 		virtual void OnValidateProfile() {};
 
 		/**
-		 * Determine if this extension also provides items in the bespoke extensions menu
+		 * Describes how this extension is represented in the extensions menu
 		 */
-		virtual bool HasCustomExtensionMenu() const { return false; }
-
-		/**
-		 * Get the custom extension menu section. Default is the extension's friendly name
-		 */
-		UE_API virtual FText GetCustomExtensionMenuSection() const;
+		struct FExtensionsMenuEntry
+		{
+			// internal names
+			FString SectionName; // if this is empty, item will be placed in the common section & SectionDisplayName will be ignored
+			FString SubmenuName;
+
+			// friendly names
+			FText SectionDisplayName;
+			FText SubmenuDisplayName;
+
+			// type of menu entry
+			enum EType
+			{
+				Type_None,		// does not appear in the extensions menu
+				Type_Section,	// appears in the extensions menu in the SectionName section.
+				Type_SubMenu,	// appears in the extensions menu in a submenu of the SectionName section
+			};
+			EType Type;
+		};
+		UE_API virtual void GetExtensionsMenuEntry( FExtensionsMenuEntry& MenuEntry ) const;
 
 		/**
 		 * Create the default heading for this extension
 		 */
 		UE_API FLaunchProfileTreeNode& AddDefaultHeading(ProjectLauncher::FLaunchProfileTreeData& ProfileTreeData) const;
 
-
 		/**
 		 * Create the default heading for this extension, with a bespoke display name
 		 */
@@ -116,7 +129,7 @@ namespace ProjectLauncher
 		 * Populate the custom extension menu for this extension. This menu would typically contain
 		 * items that might enable bespoke tree customization options for example.
 		 */
-		virtual void MakeCustomExtensionSubmenu(FMenuBuilder& MenuBuilder) { checkNoEntry(); }
+		virtual void MakeCustomExtensionSubmenu(FMenuBuilder& MenuBuilder) {}
 
 		/** 
 		 * Get the extension that instantiated us
@@ -351,15 +364,9 @@ namespace ProjectLauncher
 		UE_API virtual void MakeCustomExtensionSubmenu( FMenuBuilder& MenuBuilder ) override;
 
 		/**
-		 * Get the custom extension menu section. Default is a common section for all automated tests
-		 */
-		UE_API FText GetCustomExtensionMenuSection() const override;
-
-		/**
-		 * Returns the common menu section name for all automated tests
-		 * Can be used to insert things into this menu that may not inherit from this class
+		 * Describes how this extension is represented in the extensions menu - a common section for all automated tests
 		 */
-		UE_API static FText GetDefaultExtensionMenuSection();
+		UE_API virtual void GetExtensionsMenuEntry( FExtensionsMenuEntry& MenuEntry ) const override;
 		
 		/**
 		 * Determine if this automated test is currently active
diff --git a/Engine/Plugins/Developer/ProjectLauncher/Source/ProjectLauncher/Public/ProfileTree/LaunchProfileTreeData.h b/Engine/Plugins/Developer/ProjectLauncher/Source/ProjectLauncher/Public/ProfileTree/LaunchProfileTreeData.h
--- a/Engine/Plugins/Developer/ProjectLauncher/Source/ProjectLauncher/Public/ProfileTree/LaunchProfileTreeData.h
+++ b/Engine/Plugins/Developer/ProjectLauncher/Source/ProjectLauncher/Public/ProfileTree/LaunchProfileTreeData.h
@@ -188,7 +188,6 @@ namespace ProjectLauncher
 		TMap<FString,FLaunchProfileTreeNodePtr> HeadingNodes;            // named top-level headings
 		ILaunchProfileTreeBuilder* TreeBuilder;                          // UI builder interface
 		TArray<TSharedPtr<FLaunchExtensionInstance>> ExtensionInstances; // extension instances attached
-		bool bHasAnyMenuExtensions = false;                              // whether any extensions add menus
 		bool bRequestTreeRefresh = false;                                // request to refresh tree UI
 	};
 
