Fixed crash in Project Launcher when unplugging/disconnecting a mobile device.

Also fixed a couple copy/paste errors resuling in functions named "...ChangedChanged".

diff --git a/Engine/Plugins/Developer/ProjectLauncher/Source/ProjectLauncher/Private/Widgets/Shared/SCustomLaunchDeviceCombo.cpp b/Engine/Plugins/Developer/ProjectLauncher/Source/ProjectLauncher/Private/Widgets/Shared/SCustomLaunchDeviceCombo.cpp
--- a/Engine/Plugins/Developer/ProjectLauncher/Source/ProjectLauncher/Private/Widgets/Shared/SCustomLaunchDeviceCombo.cpp
+++ b/Engine/Plugins/Developer/ProjectLauncher/Source/ProjectLauncher/Private/Widgets/Shared/SCustomLaunchDeviceCombo.cpp
@@ -36,7 +36,7 @@ void SCustomLaunchDeviceCombo::Construct(const FArguments& InArgs)
 			SAssignNew(DeviceProxyComboBox, SComboBox<TSharedPtr<ITargetDeviceProxy>>)
 			.OptionsSource(&DeviceProxyList)
 			.OnGenerateWidget(this, &SCustomLaunchDeviceCombo::GenerateDeviceProxyListWidget)
-			.OnSelectionChanged( this, &SCustomLaunchDeviceCombo::OnDeviceProxySelectionChangedChanged )
+			.OnSelectionChanged( this, &SCustomLaunchDeviceCombo::OnDeviceProxySelectionChanged)
 			[
 				SNew(SHorizontalBox)
 					
@@ -150,21 +150,17 @@ FText SCustomLaunchDeviceCombo::GetSelectedDeviceProxyName() const
 	return LOCTEXT("NoDevice", "(no device)");
 }
 
-
-
-
-
-
-void SCustomLaunchDeviceCombo::OnDeviceProxySelectionChangedChanged(TSharedPtr<ITargetDeviceProxy> DeviceProxy, ESelectInfo::Type InSelectInfo)
+void SCustomLaunchDeviceCombo::OnDeviceProxySelectionChanged(TSharedPtr<ITargetDeviceProxy> DeviceProxy, ESelectInfo::Type InSelectInfo)
 {
-	FString DeviceID = DeviceProxy->GetTargetDeviceId(NAME_None);
-
 	TArray<FString> DeviceIDs;
-	DeviceIDs.Add( DeviceID );
+
+	if (DeviceProxy != nullptr)
+	{
+		FString DeviceID = DeviceProxy->GetTargetDeviceId(NAME_None);
+		DeviceIDs.Add(DeviceID);
+	}
 
 	OnSelectionChanged.ExecuteIfBound(DeviceIDs);
 }
 
-
-
 #undef LOCTEXT_NAMESPACE
diff --git a/Engine/Plugins/Developer/ProjectLauncher/Source/ProjectLauncher/Public/Widgets/Shared/SCustomLaunchDeviceCombo.h b/Engine/Plugins/Developer/ProjectLauncher/Source/ProjectLauncher/Public/Widgets/Shared/SCustomLaunchDeviceCombo.h
--- a/Engine/Plugins/Developer/ProjectLauncher/Source/ProjectLauncher/Public/Widgets/Shared/SCustomLaunchDeviceCombo.h
+++ b/Engine/Plugins/Developer/ProjectLauncher/Source/ProjectLauncher/Public/Widgets/Shared/SCustomLaunchDeviceCombo.h
@@ -29,7 +29,7 @@ public:
 protected:
 
 	UE_API TSharedRef<SWidget> GenerateDeviceProxyListWidget(TSharedPtr<ITargetDeviceProxy> DeviceProxy) const;
-	UE_API void OnDeviceProxySelectionChangedChanged(TSharedPtr<ITargetDeviceProxy> DeviceProxy, ESelectInfo::Type InSelectInfo);
+	UE_API void OnDeviceProxySelectionChanged(TSharedPtr<ITargetDeviceProxy> DeviceProxy, ESelectInfo::Type InSelectInfo);
 	UE_API const FSlateBrush* GetSelectedDeviceProxyBrush() const;
 	UE_API FText GetSelectedDeviceProxyName() const;
 
