fix for the incorrect platform variant being selected when trying to launch a platform variant zen streaming build

#rnx

diff --git a/Engine/Plugins/Developer/ProjectLauncher/Source/ProjectLauncher/Private/Model/ProjectLauncherModel.cpp b/Engine/Plugins/Developer/ProjectLauncher/Source/ProjectLauncher/Private/Model/ProjectLauncherModel.cpp
--- a/Engine/Plugins/Developer/ProjectLauncher/Source/ProjectLauncher/Private/Model/ProjectLauncherModel.cpp
+++ b/Engine/Plugins/Developer/ProjectLauncher/Source/ProjectLauncher/Private/Model/ProjectLauncherModel.cpp
@@ -427,6 +427,55 @@ namespace ProjectLauncher
 	}
 
 
+	void FModel::UpdateDeployDeviceProxyFromCookedPlatforms(const ILauncherProfileRef& Profile)
+	{
+		ILauncherDeviceGroupPtr DeviceGroup = Profile->GetDeployedDeviceGroup();
+		if (DeviceGroup.IsValid())
+		{
+			TSet<FString> DeviceIDs( DeviceGroup->GetDeviceIDs() );
+
+			TArray<FString> Platforms = Profile->GetCookedPlatforms();
+
+			TSet<FString> NewDeviceIDs;
+			for ( const FString& DeviceID : DeviceIDs)
+			{
+				TSharedPtr<ITargetDeviceProxy> DeviceProxy = DeviceProxyManager->FindProxyDeviceForTargetDevice(DeviceID);
+				if (DeviceProxy.IsValid())
+				{
+					TArray<FName> Variants;
+					DeviceProxy->GetVariants(Variants);
+
+					for (const FString& Platform : Platforms)
+					{
+						for( const FName& Variant : Variants)
+						{
+							FString TargetPlatformName = DeviceProxy->GetTargetPlatformName(Variant);
+							if (Platform == TargetPlatformName)
+							{
+								const TSet<FString>& TargetDeviceIds = DeviceProxy->GetTargetDeviceIds(Variant);
+								NewDeviceIDs.Append(TargetDeviceIds);
+							}
+						}
+					}
+				}
+				else
+				{
+					NewDeviceIDs.Add(DeviceID);
+				}
+			}
+
+			// update the device group if there's been changes
+			if (NewDeviceIDs.Difference(DeviceIDs).Num() > 0 || DeviceIDs.Difference(NewDeviceIDs).Num() > 0)
+			{
+				DeviceGroup->RemoveAllDevices();
+
+				for (const FString& NewDeviceID : NewDeviceIDs)
+				{
+					DeviceGroup->AddDevice(NewDeviceID);
+				}
+			}
+		}
+	}
 
 	void FModel::UpdatedCookedPlatformsFromDeployDeviceProxy(const ILauncherProfileRef& Profile, TSharedPtr<ITargetDeviceProxy> DeviceProxy)
 	{
diff --git a/Engine/Plugins/Developer/ProjectLauncher/Source/ProjectLauncher/Private/Widgets/SProjectLauncher.cpp b/Engine/Plugins/Developer/ProjectLauncher/Source/ProjectLauncher/Private/Widgets/SProjectLauncher.cpp
--- a/Engine/Plugins/Developer/ProjectLauncher/Source/ProjectLauncher/Private/Widgets/SProjectLauncher.cpp
+++ b/Engine/Plugins/Developer/ProjectLauncher/Source/ProjectLauncher/Private/Widgets/SProjectLauncher.cpp
@@ -69,6 +69,9 @@ void SProjectLauncher::OnProfileRun(const ILauncherProfilePtr& Profile)
 	LauncherProfile = Profile;
 	if (LauncherProfile.IsValid())
 	{
+		Model->UpdateDeployDeviceProxyFromCookedPlatforms(LauncherProfile.ToSharedRef());
+
+
 		LauncherWorker = Model->GetLauncher()->Launch(Model->GetDeviceProxyManager(), LauncherProfile.ToSharedRef());
 	
 		if (LauncherWorker.IsValid())
diff --git a/Engine/Plugins/Developer/ProjectLauncher/Source/ProjectLauncher/Public/Model/ProjectLauncherModel.h b/Engine/Plugins/Developer/ProjectLauncher/Source/ProjectLauncher/Public/Model/ProjectLauncherModel.h
--- a/Engine/Plugins/Developer/ProjectLauncher/Source/ProjectLauncher/Public/Model/ProjectLauncherModel.h
+++ b/Engine/Plugins/Developer/ProjectLauncher/Source/ProjectLauncher/Public/Model/ProjectLauncherModel.h
@@ -235,6 +235,9 @@ namespace ProjectLauncher
 		/* Get a device proxy for a profile */
 		static UE_API TSharedPtr<ITargetDeviceProxy> GetDeviceProxy(const ILauncherProfileRef& Profile);
 
+		/* Update device proxies based on the current cook platforms */
+		UE_API void UpdateDeployDeviceProxyFromCookedPlatforms(const ILauncherProfileRef& Profile);
+
 		/* Update cooked platforms based on selected device proxy */
 		UE_API void UpdatedCookedPlatformsFromDeployDeviceProxy(const ILauncherProfileRef& Profile, TSharedPtr<ITargetDeviceProxy> DeviceProxy = nullptr);
 
diff --git a/Engine/Source/Developer/LauncherServices/Private/Launcher/LauncherWorker.cpp b/Engine/Source/Developer/LauncherServices/Private/Launcher/LauncherWorker.cpp
--- a/Engine/Source/Developer/LauncherServices/Private/Launcher/LauncherWorker.cpp
+++ b/Engine/Source/Developer/LauncherServices/Private/Launcher/LauncherWorker.cpp
@@ -1334,6 +1334,13 @@ void FLauncherWorker::CreateAndExecuteTasks( const ILauncherProfileRef& InProfil
 		}
 	}
 
+	// could not get a platform from the device proxy (suggesting to use the default device) so use the selected cook platforms
+	if (Platforms.Num() == 0)
+	{
+		Platforms = InProfile->GetCookedPlatforms();
+	}
+
+
 #if !WITH_EDITOR
 	// can't cook by the book in the editor if we are not in the editor...
 	check( InProfile->GetCookMode() != ELauncherProfileCookModes::ByTheBookInEditor );
