Project Launcher fixes for recent changes in Cooker behavior:
- Zen Streaming content scheme is only used for new profiles when bUseZenStore and AllowRemoteNetworkService are set
- Add profile validation error when trying to use Zen Streaming to a remote device when AllowRemoteNetworkService is false
- Platforms menu Quick Launch disables unversioned cook when using LegacyIterative

diff --git a/Engine/Plugins/Developer/ProjectLauncher/Source/ProjectLauncher/Private/Model/ProjectLauncherModel.cpp b/Engine/Plugins/Developer/ProjectLauncher/Source/ProjectLauncher/Private/Model/ProjectLauncherModel.cpp
--- a/Engine/Plugins/Developer/ProjectLauncher/Source/ProjectLauncher/Private/Model/ProjectLauncherModel.cpp
+++ b/Engine/Plugins/Developer/ProjectLauncher/Source/ProjectLauncher/Private/Model/ProjectLauncherModel.cpp
@@ -390,7 +390,15 @@ namespace ProjectLauncher
 		ILauncherProfileRef Profile = ProfileManager->CreateUnsavedProfile(Name);
 
 		// set defaults
-		SetProfileContentScheme(EContentScheme::ZenStreaming, Profile );
+		FProjectSettings ProjectSettings = GetProjectSettings(Profile->GetProjectPath());
+		if (ProjectSettings.bUseZenStore && ProjectSettings.bAllowRemoteNetworkService)
+		{
+			SetProfileContentScheme(EContentScheme::ZenStreaming, Profile );
+		}
+		else
+		{
+			SetProfileContentScheme(EContentScheme::PakFiles, Profile );
+		}
 		Profile->SetBuildConfiguration(EBuildConfiguration::Development);
 		Profile->SetLaunchMode(ELauncherProfileLaunchModes::DefaultRole);
 		Profile->SetBuildMode(ELauncherProfileBuildModes::Auto);
diff --git a/Engine/Plugins/Developer/ProjectLauncher/Source/ProjectLauncher/Private/ProfileTree/GenericProfileTreeBuilder.cpp b/Engine/Plugins/Developer/ProjectLauncher/Source/ProjectLauncher/Private/ProfileTree/GenericProfileTreeBuilder.cpp
--- a/Engine/Plugins/Developer/ProjectLauncher/Source/ProjectLauncher/Private/ProfileTree/GenericProfileTreeBuilder.cpp
+++ b/Engine/Plugins/Developer/ProjectLauncher/Source/ProjectLauncher/Private/ProfileTree/GenericProfileTreeBuilder.cpp
@@ -234,7 +234,7 @@ namespace ProjectLauncher
 		HeadingNode.AddWidget( LOCTEXT("ContentSchemeLabel", "Content Scheme"), 
 			{
 				.Validation = FValidation( {ELauncherProfileValidationErrors::CookOnTheFlyDoesntSupportServer, ELauncherProfileValidationErrors::ShippingDoesntSupportCommandlineOptionsCantUseCookOnTheFly, ELauncherProfileValidationErrors::ZenPakStreamingRequiresDeployAndLaunch}, 
-											{TEXT("Validation_ZenNeedsIoStore")} ),
+											{TEXT("Validation_ZenNeedsIoStore"), TEXT("Validation_ZenNeedsRemotePermission")} ),
 			},
 			SNew(SCustomLaunchContentSchemeCombo)
 			.OnSelectionChanged(this, &FGenericProfileTreeBuilder::SetContentScheme)
@@ -603,9 +603,17 @@ namespace ProjectLauncher
 	void FGenericProfileTreeBuilder::OnValidateProfile()
 	{
 		// verify that IoStore is available when a Zen Workflow is in use
-		if ((ContentScheme == EContentScheme::ZenStreaming || ContentScheme == EContentScheme::ZenPakStreaming) && !Model->GetProjectSettings(Profile).bUseIoStore)
+		if ((ContentScheme == EContentScheme::ZenStreaming || ContentScheme == EContentScheme::ZenPakStreaming))
 		{
-			Profile->AddCustomError(TEXT("Validation_ZenNeedsIoStore"), LOCTEXT("Error_ZenNeedsIoStore", "Zen Workflows require IoStore to be enabled in your project's Packaging Settings. It cannot be overridden in a profile"));
+			if (!Model->GetProjectSettings(Profile).bUseIoStore)
+			{
+				Profile->AddCustomError(TEXT("Validation_ZenNeedsIoStore"), LOCTEXT("Error_ZenNeedsIoStore", "Zen Workflows require IoStore to be enabled in your project's Packaging Settings. It cannot be overridden in a profile"));
+			}
+			
+			if (!Model->GetProjectSettings(Profile).bAllowRemoteNetworkService && !Model->IsHostPlatform(Profile))
+			{
+				Profile->AddCustomError(TEXT("Validation_ZenNeedsRemotePermission"), LOCTEXT("Error_ZenNeedsRemotePermission", "Zen Workflows for remote devices need [Zen.AutoLaunch]:AllowRemoteNetworkService=true. It cannot be overridden in a profile"));
+			}
 		}
 	}
 
diff --git a/Engine/Source/Editor/UnrealEd/Private/PlayLevelViaLauncher.cpp b/Engine/Source/Editor/UnrealEd/Private/PlayLevelViaLauncher.cpp
--- a/Engine/Source/Editor/UnrealEd/Private/PlayLevelViaLauncher.cpp
+++ b/Engine/Source/Editor/UnrealEd/Private/PlayLevelViaLauncher.cpp
@@ -18,6 +18,7 @@
 #include "ITargetDeviceServicesModule.h"
 #include "Logging/MessageLog.h"
 #include "Misc/CoreMisc.h"
+#include "Misc/ConfigCacheIni.h"
 #include "PlatformInfo.h"
 #include "PlayLevel.h"
 #include "Settings/EditorExperimentalSettings.h"
@@ -245,6 +246,9 @@ void UEditorEngine::StartPlayUsingLauncherSession(FRequestPlaySessionParams& InR
 		CurrentLauncherCookMode = bInEditorCooking ? ELauncherProfileCookModes::OnTheFlyInEditor : ELauncherProfileCookModes::OnTheFly;
 	}
 
+	bool bLegacyIterative = false;
+	GConfig->GetBool(TEXT("CookSettings"), TEXT("LegacyIterative"), bLegacyIterative, GEditorIni);
+
 	bool bIncrementalCooking = (CookerSettings.bIterativeCookingForLaunchOn || ExperimentalSettings.bSharedCookedBuilds) && !bCookOnTheFly;
 
 	if (CurrentLauncherCookMode == ELauncherProfileCookModes::OnTheFlyInEditor ||
@@ -293,7 +297,7 @@ void UEditorEngine::StartPlayUsingLauncherSession(FRequestPlaySessionParams& InR
 		}
 	}
 	LauncherProfile->SetCookMode(CurrentLauncherCookMode);
-	LauncherProfile->SetUnversionedCooking(!bIncrementalCooking); // Unversioned cooking is not allowed with incremental cooking
+	LauncherProfile->SetUnversionedCooking(!bIncrementalCooking || !bLegacyIterative); // Unversioned cooking is not allowed with legacy incremental cooking
 	LauncherProfile->SetIncrementalCooking(bIncrementalCooking);
 	SetCookOption(TEXTVIEW("-IgnoreIniSettingsOutOfDate"), bIncrementalCooking && CookerSettings.bIgnoreIniSettingsOutOfDateForIteration);
 	SetCookOption(TEXTVIEW("-IgnoreScriptPackagesOutOfDate"), bIncrementalCooking && CookerSettings.bIgnoreScriptPackagesOutOfDateForIteration);
