Project Launcher 2 manages DefaultDeployDevice

diff --git a/Engine/Plugins/Developer/ProjectLauncher/Source/ProjectLauncher/Private/ProfileTree/GenericProfileTreeBuilder.cpp b/Engine/Plugins/Developer/ProjectLauncher/Source/ProjectLauncher/Private/ProfileTree/GenericProfileTreeBuilder.cpp
--- a/Engine/Plugins/Developer/ProjectLauncher/Source/ProjectLauncher/Private/ProfileTree/GenericProfileTreeBuilder.cpp
+++ b/Engine/Plugins/Developer/ProjectLauncher/Source/ProjectLauncher/Private/ProfileTree/GenericProfileTreeBuilder.cpp
@@ -99,10 +99,16 @@ namespace ProjectLauncher
 		const TArray<FString>& DeviceIDs = Profile->GetDeployedDeviceGroup()->GetDeviceIDs();
 
 		MapOption = Profile->GetCookedMaps().Num() > 0 ? EMapOption::Selected : EMapOption::Startup;
-		DeployDeviceOption = DeviceIDs.Num() > 0 && !DeviceIDs[0].IsEmpty() ? EDeployDeviceOption::Selected : EDeployDeviceOption::Default;
 		bShouldCook = GetCook();
 		ContentScheme = Model->DetermineProfileContentScheme(Profile);
 		CachedBuildTargetType = Model->GetBuildTargetInfo(Profile).Type;
+
+		// ensure the default deploy platform matches the cook platform because Legacy Project Launcher can create a mismatch
+		if (Profile->GetDefaultDeployPlatform() != NAME_None && !GetSelectedPlatforms().Contains(Profile->GetDefaultDeployPlatform().ToString()) )
+		{
+			FString DeployPlatform = (GetSelectedPlatforms().Num() > 0) ? GetSelectedPlatforms()[0] : FPlatformProperties::IniPlatformName();
+			Profile->SetDefaultDeployPlatform(FName(*DeployPlatform));
+		}
 	
 		CacheArchitectures();
 	}
@@ -706,6 +712,13 @@ namespace ProjectLauncher
 				Profile->AddCookedPlatform(Platform);
 			}
 		}
+
+		if (Profile->GetDefaultDeployPlatform() != NAME_None)
+		{
+			FString DeployPlatform = (GetSelectedPlatforms().Num() > 0) ? GetSelectedPlatforms()[0] : FPlatformProperties::IniPlatformName();
+			Profile->SetDefaultDeployPlatform(FName(*DeployPlatform));
+		}
+
 		OnPropertyChanged();
 
 		DeployDeviceListView->OnSelectedPlatformChanged();
@@ -1221,17 +1234,17 @@ namespace ProjectLauncher
 
 	FGenericProfileTreeBuilder::EDeployDeviceOption FGenericProfileTreeBuilder::GetDeployDeviceOption() const
 	{
-		return DeployDeviceOption;
+		return (Profile->GetDefaultDeployPlatform() == NAME_None) ? EDeployDeviceOption::Selected : EDeployDeviceOption::Default;;
 	}
 
 	void FGenericProfileTreeBuilder::SetDeployDeviceOption( EDeployDeviceOption NewDeployDeviceOption )
 	{
 		bool bShow = (NewDeployDeviceOption == EDeployDeviceOption::Selected);
 
-		DeployDeviceOption = NewDeployDeviceOption;
-
 		if (bShow)
 		{
+			Profile->SetDefaultDeployPlatform(NAME_None);
+
 			// restore the deployed device list again, if any
 			if (CachedDeployDeviceIDs.Num() > 0 && Profile->GetDeployedDeviceGroup()->GetDeviceIDs().Num() == 0)
 			{
@@ -1244,6 +1257,9 @@ namespace ProjectLauncher
 			// to set the 'default' deploy option, its necessary to remove all the devices - take a copy of the values to allow it to be restored
 			CachedDeployDeviceIDs = Profile->GetDeployedDeviceGroup()->GetDeviceIDs();
 			SetDeployDeviceIDs(TArray<FString>());
+
+			FString DeployPlatform = (GetSelectedPlatforms().Num() > 0) ? GetSelectedPlatforms()[0] : FPlatformProperties::IniPlatformName();
+			Profile->SetDefaultDeployPlatform(FName(*DeployPlatform));
 		}
 
 		OnPropertyChanged();
diff --git a/Engine/Plugins/Developer/ProjectLauncher/Source/ProjectLauncher/Public/ProfileTree/GenericProfileTreeBuilder.h b/Engine/Plugins/Developer/ProjectLauncher/Source/ProjectLauncher/Public/ProfileTree/GenericProfileTreeBuilder.h
--- a/Engine/Plugins/Developer/ProjectLauncher/Source/ProjectLauncher/Public/ProfileTree/GenericProfileTreeBuilder.h
+++ b/Engine/Plugins/Developer/ProjectLauncher/Source/ProjectLauncher/Public/ProfileTree/GenericProfileTreeBuilder.h
@@ -102,8 +102,6 @@ namespace ProjectLauncher
 			Default,
 			Selected,
 		};
-		EDeployDeviceOption DeployDeviceOption = EDeployDeviceOption::Default;
-
 
 		UE_API FString GetProjectPath() const;
 		UE_API void SetProjectName(FString ProjectPath);
