Show the platform flavor name in the platform picker, so it matches what is selected in the combo box.
e.g "Android (DXC)" not just "Android"

#rnx

diff --git a/Engine/Plugins/Developer/ProjectLauncher/Source/ProjectLauncher/Private/Model/ProjectLauncherModel.cpp b/Engine/Plugins/Developer/ProjectLauncher/Source/ProjectLauncher/Private/Model/ProjectLauncherModel.cpp
--- a/Engine/Plugins/Developer/ProjectLauncher/Source/ProjectLauncher/Private/Model/ProjectLauncherModel.cpp
+++ b/Engine/Plugins/Developer/ProjectLauncher/Source/ProjectLauncher/Private/Model/ProjectLauncherModel.cpp
@@ -296,6 +296,30 @@ namespace ProjectLauncher
 		return PlatformName;
 	}
 
+	FString FModel::GetPlatformNameWithFlavor( const FString& PlatformName )
+	{
+		const PlatformInfo::FTargetPlatformInfo* PlatformInfo = PlatformInfo::FindPlatformInfo(*PlatformName);
+		if (PlatformInfo != nullptr && PlatformInfo->VanillaInfo != nullptr)
+		{
+			if (PlatformInfo->IsFlavor())
+			{
+				for (const PlatformInfo::FTargetPlatformInfo* PlatformFlavorInfo : PlatformInfo->VanillaInfo->Flavors)
+				{
+					if (PlatformFlavorInfo->PlatformFlavor == PlatformInfo->PlatformFlavor && PlatformFlavorInfo->PlatformType == EBuildTargetType::Game)
+					{
+						return PlatformFlavorInfo->Name.ToString();
+					}
+				}
+			}
+			else
+			{
+				return PlatformInfo->VanillaInfo->Name.ToString();
+			}
+		}
+
+		return PlatformName;
+	}
+
 	FString FModel::GetBuildTargetPlatformName( const FString& PlatformName, const FTargetInfo& BuildTargetInfo )
 	{
 		const PlatformInfo::FTargetPlatformInfo* PlatformInfo = ProjectLauncher::FModel::GetPlatformInfo(*PlatformName, BuildTargetInfo);
diff --git a/Engine/Plugins/Developer/ProjectLauncher/Source/ProjectLauncher/Private/ProfileTree/GenericProfileTreeBuilder.cpp b/Engine/Plugins/Developer/ProjectLauncher/Source/ProjectLauncher/Private/ProfileTree/GenericProfileTreeBuilder.cpp
--- a/Engine/Plugins/Developer/ProjectLauncher/Source/ProjectLauncher/Private/ProfileTree/GenericProfileTreeBuilder.cpp
+++ b/Engine/Plugins/Developer/ProjectLauncher/Source/ProjectLauncher/Private/ProfileTree/GenericProfileTreeBuilder.cpp
@@ -734,7 +734,7 @@ namespace ProjectLauncher
 		{
 			for (const FString& Platform : Profile->GetCookedPlatforms())
 			{
-				Platforms.Add(FModel::GetVanillaPlatformName(Platform));
+				Platforms.Add(FModel::GetPlatformNameWithFlavor(Platform));
 			}
 		}
 		else
diff --git a/Engine/Plugins/Developer/ProjectLauncher/Source/ProjectLauncher/Public/Model/ProjectLauncherModel.h b/Engine/Plugins/Developer/ProjectLauncher/Source/ProjectLauncher/Public/Model/ProjectLauncherModel.h
--- a/Engine/Plugins/Developer/ProjectLauncher/Source/ProjectLauncher/Public/Model/ProjectLauncherModel.h
+++ b/Engine/Plugins/Developer/ProjectLauncher/Source/ProjectLauncher/Public/Model/ProjectLauncherModel.h
@@ -246,6 +246,7 @@ namespace ProjectLauncher
 
 		/* Utility functions for resolving platform names */
 		static UE_API FString GetVanillaPlatformName(const FString& PlatformName);
+		static UE_API FString GetPlatformNameWithFlavor(const FString& PlatformName);
 		static UE_API FString GetBuildTargetPlatformName(const FString& PlatformName, const FTargetInfo& BuildTargetInfo);
 
 		/* Retrieve project settings from a project path or profile */
