Add an "advanced cook" extension to Project Launcher
- less common cook items have been moved there
- added a culture selection widget as well

#rnx

diff --git a/Engine/Plugins/Developer/ProjectLauncher/Source/CommonLaunchExtensions/CommonLaunchExtensions.Build.cs b/Engine/Plugins/Developer/ProjectLauncher/Source/CommonLaunchExtensions/CommonLaunchExtensions.Build.cs
--- a/Engine/Plugins/Developer/ProjectLauncher/Source/CommonLaunchExtensions/CommonLaunchExtensions.Build.cs
+++ b/Engine/Plugins/Developer/ProjectLauncher/Source/CommonLaunchExtensions/CommonLaunchExtensions.Build.cs
@@ -16,7 +16,8 @@ public class CommonLaunchExtensions : ModuleRules
 			"Sockets",
 			"TraceLog",
 			"ProjectLauncher",
-			"DeveloperSettings"
+			"DeveloperSettings",
+			"ToolWidgets",
 		]);
 
 
diff --git a/Engine/Plugins/Developer/ProjectLauncher/Source/CommonLaunchExtensions/Private/AdvancedCook/AdvancedCookLaunchExtension.cpp b/Engine/Plugins/Developer/ProjectLauncher/Source/CommonLaunchExtensions/Private/AdvancedCook/AdvancedCookLaunchExtension.cpp
new file mode 100644
--- /dev/null
+++ b/Engine/Plugins/Developer/ProjectLauncher/Source/CommonLaunchExtensions/Private/AdvancedCook/AdvancedCookLaunchExtension.cpp
@@ -0,0 +1,450 @@
+// Copyright Epic Games, Inc. All Rights Reserved.
+
+#include "AdvancedCook/AdvancedCookLaunchExtension.h"
+#include "Widgets/Shared/SCustomLaunchMapListView.h"
+#include "Widgets/Shared/SCustomLaunchCultureListView.h"
+#include "Widgets/Input/SSegmentedControl.h"
+#include "Widgets/Text/STextBlock.h"
+#include "Widgets/Layout/SBox.h"
+#include "Widgets/SBoxPanel.h"
+#include "SResizeBox.h"
+#include "SlateOptMacros.h"
+
+#define LOCTEXT_NAMESPACE "FAdvancedCookLaunchExtensionInstance"
+
+
+
+FAdvancedCookLaunchExtensionInstance::FAdvancedCookLaunchExtensionInstance( FArgs& InArgs ) 
+	: ProjectLauncher::FLaunchExtensionInstance(InArgs)
+{
+	MapOption = GetProfile()->GetCookedMaps().Num() > 0 ? EMapOption::Selected : EMapOption::Startup;
+
+	CachedCulturesToCook = GetProfile()->GetCookedCultures();
+	CultureOption = (CachedCulturesToCook.Num() == 1 && CachedCulturesToCook.Contains(TEXT("en"))) ? ECultureOption::Default : ECultureOption::Custom;
+}
+
+
+bool FAdvancedCookLaunchExtensionInstance::IsAvailable() const
+{
+	return (GetConfigBool(EConfig::PerProfile, TEXT("Enabled") ) || HasAdvancedCookOptions());
+}
+
+
+void FAdvancedCookLaunchExtensionInstance::MakeCustomExtensionSubmenu( FMenuBuilder& MenuBuilder )
+{
+	auto IsEnabled = [this]()
+	{
+		return GetConfigBool(EConfig::PerProfile, TEXT("Enabled") );
+	};
+
+	auto ToggleEnabled = [this,IsEnabled]()
+	{
+		SetConfigBool(EConfig::PerProfile, TEXT("Enabled"), !IsEnabled() );
+	};
+
+	auto CanToggle = [this,IsEnabled]()
+	{
+		return !HasAdvancedCookOptions();
+	};
+
+	auto GetToolTip = [this]()
+	{
+		if (HasAdvancedCookOptions() && IsAvailable())
+		{
+			return LOCTEXT("CannotDisable","This extension cannot be disabled while there are advanced cook options configured");
+		}
+
+		return FText::GetEmpty();
+	};
+
+
+	MenuBuilder.AddMenuEntry(
+		LOCTEXT("EnableAdvancedCookLabel", "Enable Advanced Cook"),
+		TAttribute<FText>::CreateLambda(GetToolTip),
+		FSlateIcon(),
+		FUIAction(
+			FExecuteAction::CreateLambda(ToggleEnabled),
+			FCanExecuteAction::CreateLambda(CanToggle),
+			FIsActionChecked::CreateSP(this, &FAdvancedCookLaunchExtensionInstance::IsAvailable)
+		),
+		NAME_None,
+		EUserInterfaceActionType::Check);
+}
+
+bool FAdvancedCookLaunchExtensionInstance::HasAdvancedCookOptions() const
+{
+	if (!GetProfile()->IsCookingUnversioned() || !GetProfile()->GetCookOptions().IsEmpty() || GetProfile()->GetCookedMaps().Num() > 0)
+	{
+		return true;
+	}
+
+	bool bHasDefaultCultures = (GetProfile()->GetCookedCultures().Num() == 1 && GetProfile()->GetCookedCultures()[0] == TEXT("en"));
+	if (!bHasDefaultCultures)
+	{
+		return true;
+	}
+
+	// we're not using any advanced cook options
+	return false;
+}
+
+void FAdvancedCookLaunchExtensionInstance::CustomizeTree( ProjectLauncher::FLaunchProfileTreeData& ProfileTreeData )
+{
+	auto IsVisible = [this]()
+	{
+		return IsAvailable() && (GetProfile()->GetCookMode() == ELauncherProfileCookModes::ByTheBook);
+	};
+
+	//ProfileTreeData.AddHeading( TEXT("Cooking"), LOCTEXT("CookingHeading", "Cook") )
+	AddDefaultHeading(ProfileTreeData)
+		.AddBoolean( LOCTEXT("UnversionedCookLabel", "Unversioned Cook"),
+			{
+				.GetValue = [this]()			{ return GetProfile()->IsCookingUnversioned(); },
+				.SetValue = [this](bool bValue)	{ GetProfile()->SetUnversionedCooking(bValue); },
+				.GetDefaultValue = [this]()		{ return GetDefaultProfile()->IsCookingUnversioned(); },
+				.IsVisible = IsVisible,
+				.Validation = ProjectLauncher::FValidation({ELauncherProfileValidationErrors::UnversionedAndIncremental}),
+			}
+		)
+		.AddWidget( LOCTEXT("MapsToCookLabel", "Maps To Cook"), 
+			{
+				.IsVisible = IsVisible,
+			},
+			CreateMapListWidget()
+		)
+		.AddWidget( LOCTEXT("CulturesToCookLabel", "Cultures To Cook"), 
+			{
+				.IsVisible = IsVisible,
+			},
+			CreateCultureListWidget()
+		)
+		.AddString( LOCTEXT("AdditionalCookerOptionsLabel", "Additional Cooker Options"), 
+			{
+				.GetValue = [this]()				{ return GetProfile()->GetCookOptions(); },
+				.SetValue = [this](FString Value)	{ GetProfile()->SetCookOptions( Value ); },
+				.GetDefaultValue = []()				{ return FString(); },
+				.IsVisible = IsVisible,
+			}
+		)
+	;
+}
+
+
+
+
+BEGIN_SLATE_FUNCTION_BUILD_OPTIMIZATION
+TSharedRef<SWidget> FAdvancedCookLaunchExtensionInstance::CreateMapListWidget()
+{
+	MapListView = 
+		SNew(SCustomLaunchMapListView, GetModel())
+		.OnSelectionChanged(this, &FAdvancedCookLaunchExtensionInstance::SetMapsToCook)
+		.SelectedMaps(this, &FAdvancedCookLaunchExtensionInstance::GetMapsToCook)
+		.ProjectPath_Lambda( [this]() { return GetModel()->GetProjectPath(GetProfile()); })
+	;
+
+	return SNew(SVerticalBox)
+
+		// map option controls
+		+SVerticalBox::Slot()
+		.AutoHeight()
+		.Padding(0,2)
+		[
+			SNew(SHorizontalBox)
+
+			// map option
+			+SHorizontalBox::Slot()
+			.Padding(0, 0)
+			.AutoWidth()
+			.VAlign(VAlign_Center)
+			[
+				SNew(SSegmentedControl<EMapOption>)
+				.Value( this, &FAdvancedCookLaunchExtensionInstance::GetMapOption)
+				.OnValueChanged(this, &FAdvancedCookLaunchExtensionInstance::SetMapOption)
+
+				+SSegmentedControl<EMapOption>::Slot(EMapOption::Startup)
+				.HAlign(HAlign_Center)
+				.VAlign(VAlign_Center)
+				[
+					SNew(STextBlock)
+					.Text(LOCTEXT("StartupMapsLabel", "Startup Maps"))
+					.Font(FCoreStyle::Get().GetFontStyle("SmallFont"))
+				]
+
+				+SSegmentedControl<EMapOption>::Slot(EMapOption::Selected)
+				.HAlign(HAlign_Center)
+				.VAlign(VAlign_Center)
+				[
+					SNew(STextBlock)
+					.Text(LOCTEXT("SelectedMapsLabel", "Selected Maps"))
+					.Font(FCoreStyle::Get().GetFontStyle("SmallFont"))
+				]
+			]
+
+			// map selector controls (search etc)
+			+SHorizontalBox::Slot()
+			.Padding(8, 0)
+			.FillWidth(1)
+			.VAlign(VAlign_Center)
+			[
+				SNew(SBox)
+				.Visibility_Lambda([this](){ return GetMapOption() == EMapOption::Selected  ? EVisibility::Visible : EVisibility::Collapsed; } )
+				[
+					MapListView->MakeControlsWidget()
+				]
+			]
+		]
+
+		// map list
+		+SVerticalBox::Slot()
+		.FillHeight(1)
+		[
+			SNew(SVerticalResizeBox)
+			.Visibility_Lambda([this](){ return (GetMapOption() == EMapOption::Selected) ? EVisibility::Visible : EVisibility::Collapsed; } )
+			.HandleHeight(4)
+			.ContentHeight(this, &FAdvancedCookLaunchExtensionInstance::GetMapListHeight)
+			.ContentHeightChanged(this, &FAdvancedCookLaunchExtensionInstance::SetMapListHeight)
+			.HandleColor( FAppStyle::Get().GetSlateColor("Colors.Secondary").GetSpecifiedColor() )
+			[
+				MapListView.ToSharedRef()
+			]
+		]
+	;
+}
+END_SLATE_FUNCTION_BUILD_OPTIMIZATION
+
+
+
+BEGIN_SLATE_FUNCTION_BUILD_OPTIMIZATION
+TSharedRef<SWidget> FAdvancedCookLaunchExtensionInstance::CreateCultureListWidget()
+{
+	CultureListView = 
+		SNew(SCustomLaunchCultureListView, GetModel())
+		.OnSelectionChanged(this, &FAdvancedCookLaunchExtensionInstance::SetCulturesToCook)
+		.SelectedCultures(this, &FAdvancedCookLaunchExtensionInstance::GetCulturesToCook)
+		.ProjectPath_Lambda( [this]() { return GetModel()->GetProjectPath(GetProfile()); })
+	;
+
+	return SNew(SVerticalBox)
+
+		// culture option controls
+		+SVerticalBox::Slot()
+		.AutoHeight()
+		.Padding(0,2)
+		[
+			SNew(SHorizontalBox)
+
+			// culture option
+			+SHorizontalBox::Slot()
+			.Padding(0, 0)
+			.AutoWidth()
+			.VAlign(VAlign_Center)
+			[
+				SNew(SSegmentedControl<ECultureOption>)
+				.Value( this, &FAdvancedCookLaunchExtensionInstance::GetCultureOption)
+				.OnValueChanged(this, &FAdvancedCookLaunchExtensionInstance::SetCultureOption)
+
+				+SSegmentedControl<ECultureOption>::Slot(ECultureOption::Default)
+				.HAlign(HAlign_Center)
+				.VAlign(VAlign_Center)
+				[
+					SNew(STextBlock)
+					.Text(LOCTEXT("DefaultCultureLabel", "Default"))
+					.Font(FCoreStyle::Get().GetFontStyle("SmallFont"))
+				]
+
+				+SSegmentedControl<ECultureOption>::Slot(ECultureOption::Custom)
+				.HAlign(HAlign_Center)
+				.VAlign(VAlign_Center)
+				[
+					SNew(STextBlock)
+					.Text(LOCTEXT("CustomCultureLabel", "Custom"))
+					.Font(FCoreStyle::Get().GetFontStyle("SmallFont"))
+				]
+			]
+
+			// culture selector controls (search etc)
+			+SHorizontalBox::Slot()
+			.Padding(8, 0)
+			.FillWidth(1)
+			.VAlign(VAlign_Center)
+			[
+				SNew(SBox)
+				.Visibility_Lambda([this](){ return GetCultureOption() == ECultureOption::Custom ? EVisibility::Visible : EVisibility::Collapsed; } )
+				[
+					CultureListView->MakeControlsWidget()
+				]
+			]
+		]
+
+		// culture list
+		+SVerticalBox::Slot()
+		.FillHeight(1)
+		[
+			SNew(SVerticalResizeBox)
+			.Visibility_Lambda([this](){ return GetCultureOption() == ECultureOption::Custom ? EVisibility::Visible : EVisibility::Collapsed; } )
+			.HandleHeight(4)
+			.ContentHeight(this, &FAdvancedCookLaunchExtensionInstance::GetCultureListHeight)
+			.ContentHeightChanged(this, &FAdvancedCookLaunchExtensionInstance::SetCultureListHeight)
+			.HandleColor( FAppStyle::Get().GetSlateColor("Colors.Secondary").GetSpecifiedColor() )
+			[
+				CultureListView.ToSharedRef()
+			]
+		]
+	;
+}
+END_SLATE_FUNCTION_BUILD_OPTIMIZATION
+
+
+void FAdvancedCookLaunchExtensionInstance::SetMapsToCook(TArray<FString> MapsToCook)
+{
+	GetProfile()->ClearCookedMaps(); 
+	for (const FString& Map : MapsToCook)
+	{
+		GetProfile()->AddCookedMap(Map);
+	}
+
+	OnPropertyChanged();
+}
+
+TArray<FString> FAdvancedCookLaunchExtensionInstance::GetMapsToCook() const
+{
+	return GetProfile()->GetCookedMaps();
+}
+
+float FAdvancedCookLaunchExtensionInstance::GetMapListHeight() const
+{
+	return MapListHeight;
+}
+
+void FAdvancedCookLaunchExtensionInstance::SetMapListHeight( float NewHeight )
+{
+	static const float MinMapListHeight = 100.0f;
+
+	MapListHeight = FMath::Max(NewHeight, MinMapListHeight);
+
+	RequestTreeRefresh();
+}
+
+FAdvancedCookLaunchExtensionInstance::EMapOption FAdvancedCookLaunchExtensionInstance::GetMapOption() const
+{
+	return MapOption;
+}
+
+void FAdvancedCookLaunchExtensionInstance::SetMapOption( EMapOption NewMapOption )
+{
+	bool bShow = (NewMapOption == EMapOption::Selected);
+
+	MapOption = NewMapOption;
+
+	if (bShow)
+	{
+		// restore the cooked maps again, if any
+		if (CachedMapsToCook.Num() > 0 && GetProfile()->GetCookedMaps().Num() == 0)
+		{
+			SetMapsToCook(CachedMapsToCook);
+			CachedMapsToCook.Reset();
+		}
+	}
+	else
+	{
+		// to set the 'cook startup maps only', its necessary to remove all the cooked maps - take a copy of the values to allow it to be restored
+		CachedMapsToCook = GetProfile()->GetCookedMaps();
+		SetMapsToCook(TArray<FString>());
+	}
+
+	OnPropertyChanged();
+
+	MapListView->RefreshMapList();
+}
+
+
+
+
+void FAdvancedCookLaunchExtensionInstance::SetCulturesToCook(TArray<FString> CulturesToCook)
+{
+	GetProfile()->ClearCookedCultures(); 
+	for (const FString& Culture : CulturesToCook)
+	{
+		GetProfile()->AddCookedCulture(Culture);
+	}
+
+	OnPropertyChanged();
+}
+
+TArray<FString> FAdvancedCookLaunchExtensionInstance::GetCulturesToCook() const
+{
+	return GetProfile()->GetCookedCultures();
+}
+
+float FAdvancedCookLaunchExtensionInstance::GetCultureListHeight() const
+{
+	return CultureListHeight;
+}
+
+void FAdvancedCookLaunchExtensionInstance::SetCultureListHeight( float NewHeight )
+{
+	static const float MinCultureListHeight = 100.0f;
+
+	CultureListHeight = FMath::Max(NewHeight, MinCultureListHeight);
+
+	RequestTreeRefresh();
+}
+
+FAdvancedCookLaunchExtensionInstance::ECultureOption FAdvancedCookLaunchExtensionInstance::GetCultureOption() const
+{
+	return CultureOption;
+}
+
+void FAdvancedCookLaunchExtensionInstance::SetCultureOption( ECultureOption NewCultureOption )
+{
+	bool bShow = (NewCultureOption == ECultureOption::Custom);
+
+	CultureOption = NewCultureOption;
+
+	if (bShow)
+	{
+		// restore the cooked cultures again
+		SetCulturesToCook(CachedCulturesToCook);
+		CachedCulturesToCook.Reset();
+	}
+	else
+	{
+		// to set the 'english only', its necessary to remove all the currently selected cultures - take a copy of the values to allow it to be restored
+		CachedCulturesToCook = GetProfile()->GetCookedCultures();
+		SetCulturesToCook(TArray<FString>{TEXT("en")});
+	}
+
+	OnPropertyChanged();
+
+	CultureListView->RefreshCultureList();
+}
+
+
+void FAdvancedCookLaunchExtensionInstance::OnProjectChanged()
+{
+	if (MapListView.IsValid())
+	{
+		MapListView->RefreshMapList();
+	}
+}
+
+
+
+
+TSharedPtr<ProjectLauncher::FLaunchExtensionInstance> FAdvancedCookLaunchExtension::CreateInstanceForProfile( ProjectLauncher::FLaunchExtensionInstance::FArgs& InArgs )
+{
+	return MakeShared<FAdvancedCookLaunchExtensionInstance>(InArgs);
+}
+
+const TCHAR* FAdvancedCookLaunchExtension::GetInternalName() const
+{
+	return TEXT("AdvancedCook");
+}
+
+FText FAdvancedCookLaunchExtension::GetDisplayName() const
+{
+	return LOCTEXT("ExtensionName", "Advanced Cook Options");
+}
+
+
+#undef LOCTEXT_NAMESPACE
diff --git a/Engine/Plugins/Developer/ProjectLauncher/Source/CommonLaunchExtensions/Private/AdvancedCook/AdvancedCookLaunchExtension.h b/Engine/Plugins/Developer/ProjectLauncher/Source/CommonLaunchExtensions/Private/AdvancedCook/AdvancedCookLaunchExtension.h
new file mode 100644
--- /dev/null
+++ b/Engine/Plugins/Developer/ProjectLauncher/Source/CommonLaunchExtensions/Private/AdvancedCook/AdvancedCookLaunchExtension.h
@@ -0,0 +1,77 @@
+// Copyright Epic Games, Inc. All Rights Reserved.
+
+#pragma once
+
+#include "Extension/LaunchExtension.h"
+
+
+class SCustomLaunchMapListView;
+class SCustomLaunchCultureListView;
+
+
+class FAdvancedCookLaunchExtensionInstance : public ProjectLauncher::FLaunchExtensionInstance
+{
+public:
+	FAdvancedCookLaunchExtensionInstance( FArgs& InArgs );
+	virtual ~FAdvancedCookLaunchExtensionInstance() = default;
+
+	virtual bool HasCustomExtensionMenu() const { return true; }
+	virtual void MakeCustomExtensionSubmenu( FMenuBuilder& MenuBuilder ) override;
+
+	virtual void CustomizeTree( ProjectLauncher::FLaunchProfileTreeData& ProfileTreeData ) override;
+
+protected:
+	TArray<FString> CachedMapsToCook;
+	enum class EMapOption : uint8
+	{
+		Startup,
+		Selected,
+	};
+	EMapOption MapOption = EMapOption::Startup;
+
+	void SetMapsToCook(TArray<FString> MapsToCook);
+	TArray<FString> GetMapsToCook() const;
+	EMapOption GetMapOption() const;
+	void SetMapOption( EMapOption MapOption );
+	float GetMapListHeight() const;
+	void SetMapListHeight( float NewHeight );
+	TSharedRef<SWidget> CreateMapListWidget();
+
+	float MapListHeight = 300.0f;
+	TSharedPtr<SCustomLaunchMapListView> MapListView;
+
+
+	TArray<FString> CachedCulturesToCook;
+	enum class ECultureOption : uint8
+	{
+		Default,
+		Custom,
+	};
+	ECultureOption CultureOption = ECultureOption::Default;
+
+	void SetCulturesToCook(TArray<FString> CulturesToCook);
+	TArray<FString> GetCulturesToCook() const;
+	ECultureOption GetCultureOption() const;
+	void SetCultureOption( ECultureOption CultureOption );
+	float GetCultureListHeight() const;
+	void SetCultureListHeight( float NewHeight );
+	TSharedRef<SWidget> CreateCultureListWidget();
+
+	float CultureListHeight = 150.0f;
+	TSharedPtr<SCustomLaunchCultureListView> CultureListView;
+
+
+	virtual void OnProjectChanged() override;
+	bool HasAdvancedCookOptions() const;
+
+	bool IsAvailable() const;
+};
+
+
+class FAdvancedCookLaunchExtension : public ProjectLauncher::FLaunchExtension
+{
+public:
+	virtual TSharedPtr<ProjectLauncher::FLaunchExtensionInstance> CreateInstanceForProfile( ProjectLauncher::FLaunchExtensionInstance::FArgs& InArgs ) override;
+	virtual const TCHAR* GetInternalName() const override;
+	virtual FText GetDisplayName() const override;
+};
\ No newline at end of file
diff --git a/Engine/Plugins/Developer/ProjectLauncher/Source/CommonLaunchExtensions/Private/CommonLaunchExtensionsModule.cpp b/Engine/Plugins/Developer/ProjectLauncher/Source/CommonLaunchExtensions/Private/CommonLaunchExtensionsModule.cpp
--- a/Engine/Plugins/Developer/ProjectLauncher/Source/CommonLaunchExtensions/Private/CommonLaunchExtensionsModule.cpp
+++ b/Engine/Plugins/Developer/ProjectLauncher/Source/CommonLaunchExtensions/Private/CommonLaunchExtensionsModule.cpp
@@ -6,6 +6,7 @@
 #include "Globals/GlobalsLaunchExtension.h"
 #include "Insights/InsightsLaunchExtension.h"
 #include "BootTest/BootTestLaunchExtension.h"
+#include "AdvancedCook/AdvancedCookLaunchExtension.h"
 
 class FCommonLaunchExtensionsModule : public IModuleInterface
 {
@@ -13,34 +14,32 @@ public:
 
 	virtual void StartupModule() override
 	{
-		Globals = MakeShared<FGlobalsLaunchExtension>();
-		Insights = MakeShared<FInsightsLaunchExtension>();
-		BootTest = MakeShared<FBootTestLaunchExtension>();
+		Extensions.Add( MakeShared<FGlobalsLaunchExtension>() );
+		Extensions.Add( MakeShared<FInsightsLaunchExtension>() );
+		Extensions.Add( MakeShared<FBootTestLaunchExtension>() );
+		Extensions.Add( MakeShared<FAdvancedCookLaunchExtension>() );
 
-		IProjectLauncherModule::Get().RegisterExtension(Globals.ToSharedRef());
-		IProjectLauncherModule::Get().RegisterExtension(Insights.ToSharedRef());
-		IProjectLauncherModule::Get().RegisterExtension(BootTest.ToSharedRef());
+		for (TSharedRef<ProjectLauncher::FLaunchExtension> Extension : Extensions)
+		{
+			IProjectLauncherModule::Get().RegisterExtension(Extension);
+		}
 	}
 
 	virtual void ShutdownModule() override
 	{
 		if (IProjectLauncherModule* ProjectLauncher = IProjectLauncherModule::TryGet())
 		{
-			ProjectLauncher->UnregisterExtension(Globals.ToSharedRef());
-			ProjectLauncher->UnregisterExtension(Insights.ToSharedRef());
-			ProjectLauncher->UnregisterExtension(BootTest.ToSharedRef());
+			for (TSharedRef<ProjectLauncher::FLaunchExtension> Extension : Extensions)
+			{
+				ProjectLauncher->UnregisterExtension(Extension);
+			}
 		}
 
-		Globals.Reset();
-		Insights.Reset();
-		BootTest.Reset();
+		Extensions.Reset();
 	}
 
 private:
-	TSharedPtr<FGlobalsLaunchExtension> Globals;
-	TSharedPtr<FInsightsLaunchExtension> Insights;
-	TSharedPtr<FBootTestLaunchExtension> BootTest;
-
+	TArray<TSharedRef<ProjectLauncher::FLaunchExtension>> Extensions;
 };
 
 
diff --git a/Engine/Plugins/Developer/ProjectLauncher/Source/ProjectLauncher/Private/Extension/LaunchExtension.cpp b/Engine/Plugins/Developer/ProjectLauncher/Source/ProjectLauncher/Private/Extension/LaunchExtension.cpp
--- a/Engine/Plugins/Developer/ProjectLauncher/Source/ProjectLauncher/Private/Extension/LaunchExtension.cpp
+++ b/Engine/Plugins/Developer/ProjectLauncher/Source/ProjectLauncher/Private/Extension/LaunchExtension.cpp
@@ -125,7 +125,7 @@ namespace ProjectLauncher
 
 
 
-	TArray<TSharedPtr<FLaunchExtensionInstance>> FLaunchExtension::CreateExtensionInstancesForProfile( ILauncherProfileRef InProfile, TSharedRef<FModel> InModel )
+	TArray<TSharedPtr<FLaunchExtensionInstance>> FLaunchExtension::CreateExtensionInstancesForProfile( ILauncherProfileRef InProfile, TSharedRef<FModel> InModel, FLaunchProfileTreeData& OwnerTreeData )
 	{
 		// make sure the extensions for this profile haven't been created already
 		{
@@ -144,6 +144,9 @@ namespace ProjectLauncher
 			Extensions = GExtensions;
 		}
 
+		// look up the default profile
+		ILauncherProfileRef DefaultProfile = InModel->IsBasicLaunchProfile(InProfile) ? InModel->GetDefaultBasicLaunchProfile() : InModel->GetDefaultCustomLaunchProfile();
+
 		// attempt to instantiate all extensions
 		TArray<TSharedPtr<FLaunchExtensionInstance>> Result;
 		for (TSharedPtr<FLaunchExtension> Extension : Extensions)
@@ -151,8 +154,10 @@ namespace ProjectLauncher
 			FLaunchExtensionInstance::FArgs Args
 			{
 				.Profile = InProfile,
+				.DefaultProfile = DefaultProfile,
 				.Model = InModel,
 				.Extension = Extension.ToSharedRef(),
+				.OwnerTreeData = OwnerTreeData,
 			};
 
 			TSharedPtr<FLaunchExtensionInstance> ExtensionInstance = Extension->CreateInstanceForProfile(Args);
@@ -182,15 +187,19 @@ namespace ProjectLauncher
 
 	FLaunchExtensionInstance::FLaunchExtensionInstance( FArgs& InArgs )
 		: Profile(InArgs.Profile)
+		, DefaultProfile(InArgs.DefaultProfile)
 		, Model(InArgs.Model)
 		, Extension(InArgs.Extension)
+		, OwnerTreeData(InArgs.OwnerTreeData)
 	{
 		Profile->OnCustomValidation().AddRaw( this, &FLaunchExtensionInstance::OnValidateProfile);
+		Profile->OnProjectChanged().AddRaw( this, &FLaunchExtensionInstance::OnProjectChanged);
 	}
 
 	FLaunchExtensionInstance::~FLaunchExtensionInstance()
 	{
 		Profile->OnCustomValidation().RemoveAll(this);
+		Profile->OnProjectChanged().RemoveAll(this);
 	}
 
 	bool FLaunchExtensionInstance::GetExtensionVariables( TArray<FString>& OutVariables ) const
@@ -634,6 +643,11 @@ namespace ProjectLauncher
 		PropertyChangedDelegate.Broadcast();
 	}
 
+	void FLaunchExtensionInstance::RequestTreeRefresh()
+	{
+		OwnerTreeData.RequestTreeRefresh();
+	}
+
 	FLaunchProfileTreeNode& FLaunchExtensionInstance::AddDefaultHeading(ProjectLauncher::FLaunchProfileTreeData& ProfileTreeData) const
 	{
 		return AddDefaultHeading(ProfileTreeData, Extension->GetDisplayName());
diff --git a/Engine/Plugins/Developer/ProjectLauncher/Source/ProjectLauncher/Private/Model/ProjectLauncherModel.cpp b/Engine/Plugins/Developer/ProjectLauncher/Source/ProjectLauncher/Private/Model/ProjectLauncherModel.cpp
--- a/Engine/Plugins/Developer/ProjectLauncher/Source/ProjectLauncher/Private/Model/ProjectLauncherModel.cpp
+++ b/Engine/Plugins/Developer/ProjectLauncher/Source/ProjectLauncher/Private/Model/ProjectLauncherModel.cpp
@@ -128,7 +128,6 @@ namespace ProjectLauncher
 	bool FModel::IsAdvancedProfile(const ILauncherProfileRef& Profile) const
 	{
 		return Profile->GetCookedPlatforms().Num() > 1 ||
-				Profile->GetCookedCultures().Num() > 1 ||
 				Profile->IsCreatingDLC() ||
 				Profile->IsCreatingReleaseVersion() ||
 				Profile->IsGeneratingPatch() ||				
@@ -765,6 +764,31 @@ namespace ProjectLauncher
 		return GetProjectSettings(*Profile->GetProjectPath());
 	}
 
+	FString FModel::GetProjectPath(const ILauncherProfileRef& Profile) const
+	{
+		if (IsBasicLaunchProfile(Profile))
+		{
+			if (Profile->HasProjectSpecified())
+			{
+				return Profile->GetProjectPath();
+			}
+			else
+			{
+				return GetProfileManager()->GetProjectPath();
+			}
+		}
+		else
+		{
+			if (Profile->HasProjectSpecified())
+			{
+				return Profile->GetProjectPath();
+			}
+
+			return FString();
+		}
+
+	}
+
 	bool FModel::IsBasicLaunchProfile( const ILauncherProfilePtr& Profile ) const
 	{
 		return (Profile == BasicLaunchProfile);
diff --git a/Engine/Plugins/Developer/ProjectLauncher/Source/ProjectLauncher/Private/ProfileTree/CustomProfileTreeBuilder.cpp b/Engine/Plugins/Developer/ProjectLauncher/Source/ProjectLauncher/Private/ProfileTree/CustomProfileTreeBuilder.cpp
--- a/Engine/Plugins/Developer/ProjectLauncher/Source/ProjectLauncher/Private/ProfileTree/CustomProfileTreeBuilder.cpp
+++ b/Engine/Plugins/Developer/ProjectLauncher/Source/ProjectLauncher/Private/ProfileTree/CustomProfileTreeBuilder.cpp
@@ -31,12 +31,12 @@ namespace ProjectLauncher
 		AddImportZenSnapshotProperty(ContentSchemeHeader);
 		AddZenPakStreamingPathProperty(ContentSchemeHeader);
 
-		FLaunchProfileTreeNode& MapsAndCookingHeader = TreeData->AddHeading( TEXT("Cooking"), LOCTEXT("CookingHeading", "Maps And Cooking") );
-		AddCookProperty(MapsAndCookingHeader);
-		AddIncrementalCookProperty(MapsAndCookingHeader);
-		AddUnversionedCookProperty(MapsAndCookingHeader);
-		AddMapsToCookProperty(MapsAndCookingHeader);
-		AddAdditionalCookerOptionsProperty(MapsAndCookingHeader);
+		FLaunchProfileTreeNode& CookingHeader = TreeData->AddHeading( TEXT("Cooking"), LOCTEXT("CookingHeading", "Cook") );
+		AddCookProperty(CookingHeader);
+		AddIncrementalCookProperty(CookingHeader);
+
+		FLaunchProfileTreeNode& AdvCookingHeader = TreeData->AddHeading( TEXT("AdvancedCook"), LOCTEXT("AdvancedCookingHeading", "Advanced Cooking Options") ); // @fixme: hack - just so Advanced Cooking Extension panel appears directly below normal cook. @todo: Need 'InsertAfter' for headings
+
 
 		FLaunchProfileTreeNode& BuildHeader = TreeData->AddHeading( TEXT("Build"), LOCTEXT("BuildHeading", "Build") );
 		AddBuildProperty(BuildHeader);
diff --git a/Engine/Plugins/Developer/ProjectLauncher/Source/ProjectLauncher/Private/ProfileTree/GenericProfileTreeBuilder.cpp b/Engine/Plugins/Developer/ProjectLauncher/Source/ProjectLauncher/Private/ProfileTree/GenericProfileTreeBuilder.cpp
--- a/Engine/Plugins/Developer/ProjectLauncher/Source/ProjectLauncher/Private/ProfileTree/GenericProfileTreeBuilder.cpp
+++ b/Engine/Plugins/Developer/ProjectLauncher/Source/ProjectLauncher/Private/ProfileTree/GenericProfileTreeBuilder.cpp
@@ -6,7 +6,6 @@
 #include "Widgets/Shared/SCustomLaunchPlatformCombo.h"
 #include "Widgets/Shared/SCustomLaunchContentSchemeCombo.h"
 #include "Widgets/Shared/SCustomLaunchDeviceListView.h"
-#include "Widgets/Shared/SCustomLaunchMapListView.h"
 #include "Widgets/Shared/SCustomLaunchCombo.h"
 #include "Widgets/Input/SSegmentedControl.h"
 #include "Widgets/Text/STextBlock.h"
@@ -98,7 +97,6 @@ namespace ProjectLauncher
 	{
 		const TArray<FString>& DeviceIDs = Profile->GetDeployedDeviceGroup()->GetDeviceIDs();
 
-		MapOption = Profile->GetCookedMaps().Num() > 0 ? EMapOption::Selected : EMapOption::Startup;
 		bShouldCook = GetCook();
 		ContentScheme = Model->DetermineProfileContentScheme(Profile);
 		CachedBuildTargetType = Model->GetBuildTargetInfo(Profile).Type;
@@ -350,19 +348,6 @@ namespace ProjectLauncher
 		);
 	}
 
-	void FGenericProfileTreeBuilder::AddUnversionedCookProperty( FLaunchProfileTreeNode& HeadingNode )
-	{
-		HeadingNode.AddBoolean( LOCTEXT("UnversionedCookLabel", "Unversioned Cook"),
-			{
-				.GetValue = [this]()			{ return Profile->IsCookingUnversioned(); },
-				.SetValue = [this](bool bValue)	{ Profile->SetUnversionedCooking(bValue); },
-				.GetDefaultValue = [this]()		{ return DefaultProfile->IsCookingUnversioned(); },
-				.IsVisible = ForEnabledCooked,
-				.Validation = FValidation({ELauncherProfileValidationErrors::UnversionedAndIncremental}),
-			}
-		);
-	}
-
 
 
 	void FGenericProfileTreeBuilder::AddCookProperty( FLaunchProfileTreeNode& HeadingNode )
@@ -378,28 +363,6 @@ namespace ProjectLauncher
 		);
 	}
 
-	void FGenericProfileTreeBuilder::AddMapsToCookProperty( FLaunchProfileTreeNode& HeadingNode )
-	{
-		HeadingNode.AddWidget( LOCTEXT("MapsToCookLabel", "Maps To Cook"), 
-			{
-				.IsVisible = ForEnabledCooked,
-			},
-			CreateMapListWidget()
-		);
-	}
-
-	void FGenericProfileTreeBuilder::AddAdditionalCookerOptionsProperty( FLaunchProfileTreeNode& HeadingNode )
-	{
-		HeadingNode.AddString( LOCTEXT("AdditionalCookerOptionsLabel", "Additional Cooker Options"), 
-			{
-				.GetValue = [this]()				{ return Profile->GetCookOptions(); },
-				.SetValue = [this](FString Value)	{ Profile->SetCookOptions( Value ); },
-				.GetDefaultValue = EmptyString,
-				.IsVisible = ForEnabledCooked,
-			}
-		);
-	}
-
 	void FGenericProfileTreeBuilder::AddBuildProperty( FLaunchProfileTreeNode& HeadingNode )
 	{
 		HeadingNode.AddBoolean( LOCTEXT("BuildLabel", "Build the game"),
@@ -748,26 +711,7 @@ namespace ProjectLauncher
 
 	FString FGenericProfileTreeBuilder::GetProjectPath() const
 	{
-		if (ProfileType == EProfileType::Basic)
-		{
-			if (Profile->HasProjectSpecified())
-			{
-				return Profile->GetProjectPath();
-			}
-			else
-			{
-				return Model->GetProfileManager()->GetProjectPath();
-			}
-		}
-		else
-		{
-			if (Profile->HasProjectSpecified())
-			{
-				return Profile->GetProjectPath();
-			}
-
-			return FString();
-		}
+		return Model->GetProjectPath(Profile);
 	}
 
 	void FGenericProfileTreeBuilder::SetProjectName(FString ProjectPath)
@@ -782,11 +726,6 @@ namespace ProjectLauncher
 
 		OnPropertyChanged();
 
-		if (MapListView.IsValid())
-		{
-			MapListView->RefreshMapList();
-		}
-
 		bStartupMapCacheDirty = true;
 	}
 
@@ -950,85 +889,7 @@ namespace ProjectLauncher
 
 
 
-	BEGIN_SLATE_FUNCTION_BUILD_OPTIMIZATION
-	TSharedRef<SWidget> FGenericProfileTreeBuilder::CreateMapListWidget()
-	{
-		MapListView = 
-			SNew(SCustomLaunchMapListView, Model)
-			.OnSelectionChanged(this, &FGenericProfileTreeBuilder::SetMapsToCook)
-			.SelectedMaps(this, &FGenericProfileTreeBuilder::GetMapsToCook)
-			.ProjectPath(this, &FGenericProfileTreeBuilder::GetProjectPath)
-		;
 
-		return SNew(SVerticalBox)
-
-			// map option controls
-			+SVerticalBox::Slot()
-			.AutoHeight()
-			.Padding(0,2)
-			[
-				SNew(SHorizontalBox)
-
-				// map option
-				+SHorizontalBox::Slot()
-				.Padding(0, 0)
-				.AutoWidth()
-				.VAlign(VAlign_Center)
-				[
-					SNew(SSegmentedControl<EMapOption>)
-					.Value( this, &FGenericProfileTreeBuilder::GetMapOption)
-					.OnValueChanged(this, &FGenericProfileTreeBuilder::SetMapOption)
-
-					+SSegmentedControl<EMapOption>::Slot(EMapOption::Startup)
-					.HAlign(HAlign_Center)
-					.VAlign(VAlign_Center)
-					[
-						SNew(STextBlock)
-						.Text(LOCTEXT("StartupMapsLabel", "Startup Maps"))
-						.Font(FCoreStyle::Get().GetFontStyle("SmallFont"))
-					]
-
-					+SSegmentedControl<EMapOption>::Slot(EMapOption::Selected)
-					.HAlign(HAlign_Center)
-					.VAlign(VAlign_Center)
-					[
-						SNew(STextBlock)
-						.Text(LOCTEXT("SelectedMapsLabel", "Selected Maps"))
-						.Font(FCoreStyle::Get().GetFontStyle("SmallFont"))
-					]
-				]
-
-				// map selector controls (search etc)
-				+SHorizontalBox::Slot()
-				.Padding(8, 0)
-				.FillWidth(1)
-				.VAlign(VAlign_Center)
-				[
-					SNew(SBox)
-					.Visibility_Lambda([this](){ return GetMapOption() == EMapOption::Selected  ? EVisibility::Visible : EVisibility::Collapsed; } )
-					[
-						MapListView->MakeControlsWidget()
-					]
-				]
-			]
-
-			// map list
-			+SVerticalBox::Slot()
-			.FillHeight(1)
-			[
-				SNew(SVerticalResizeBox)
-				.Visibility_Lambda([this](){ return (GetMapOption() == EMapOption::Selected) ? EVisibility::Visible : EVisibility::Collapsed; } )
-				.HandleHeight(4)
-				.ContentHeight(this, &FGenericProfileTreeBuilder::GetMapListHeight)
-				.ContentHeightChanged(this, &FGenericProfileTreeBuilder::SetMapListHeight)
-				.HandleColor( FAppStyle::Get().GetSlateColor("Colors.Secondary").GetSpecifiedColor() )
-				[
-					MapListView.ToSharedRef()
-				]
-			]
-		;
-	}
-	END_SLATE_FUNCTION_BUILD_OPTIMIZATION
 
 
 	void FGenericProfileTreeBuilder::SetCook( bool bCook )
@@ -1046,67 +907,7 @@ namespace ProjectLauncher
 		return (InProfile->GetCookMode() != ELauncherProfileCookModes::DoNotCook);
 	}
 
-	void FGenericProfileTreeBuilder::SetMapsToCook(TArray<FString> MapsToCook)
-	{
-		Profile->ClearCookedMaps(); 
-		for (const FString& Map : MapsToCook)
-		{
-			Profile->AddCookedMap(Map);
-		}
-
-		OnPropertyChanged();
-	}
-
-	TArray<FString> FGenericProfileTreeBuilder::GetMapsToCook() const
-	{
-		return Profile->GetCookedMaps();
-	}
-
-	float FGenericProfileTreeBuilder::GetMapListHeight() const
-	{
-		return MapListHeight;
-	}
-
-	void FGenericProfileTreeBuilder::SetMapListHeight( float NewHeight )
-	{
-		static const float MinMapListHeight = 100.0f;
-
-		MapListHeight = FMath::Max(NewHeight, MinMapListHeight);
-
-		TreeData->RequestTreeRefresh();
-	}
-
-	FGenericProfileTreeBuilder::EMapOption FGenericProfileTreeBuilder::GetMapOption() const
-	{
-		return MapOption;
-	}
-
-	void FGenericProfileTreeBuilder::SetMapOption( EMapOption NewMapOption )
-	{
-		bool bShow = (NewMapOption == EMapOption::Selected);
 
-		MapOption = NewMapOption;
-
-		if (bShow)
-		{
-			// restore the cooked maps again, if any
-			if (CachedMapsToCook.Num() > 0 && Profile->GetCookedMaps().Num() == 0)
-			{
-				SetMapsToCook(CachedMapsToCook);
-				CachedMapsToCook.Reset();
-			}
-		}
-		else
-		{
-			// to set the 'cook startup maps only', its necessary to remove all the cooked maps - take a copy of the values to allow it to be restored
-			CachedMapsToCook = Profile->GetCookedMaps();
-			SetMapsToCook(TArray<FString>());
-		}
-
-		OnPropertyChanged();
-
-		MapListView->RefreshMapList();
-	}
 
 
 
diff --git a/Engine/Plugins/Developer/ProjectLauncher/Source/ProjectLauncher/Private/ProfileTree/LaunchProfileTreeData.cpp b/Engine/Plugins/Developer/ProjectLauncher/Source/ProjectLauncher/Private/ProfileTree/LaunchProfileTreeData.cpp
--- a/Engine/Plugins/Developer/ProjectLauncher/Source/ProjectLauncher/Private/ProfileTree/LaunchProfileTreeData.cpp
+++ b/Engine/Plugins/Developer/ProjectLauncher/Source/ProjectLauncher/Private/ProfileTree/LaunchProfileTreeData.cpp
@@ -121,7 +121,7 @@ namespace ProjectLauncher
 
 		if (InProfile.IsValid())
 		{
-			ExtensionInstances = FLaunchExtension::CreateExtensionInstancesForProfile(Profile.ToSharedRef(), InModel);
+			ExtensionInstances = FLaunchExtension::CreateExtensionInstancesForProfile(Profile.ToSharedRef(), InModel, *this);
 
 			for (TSharedPtr<FLaunchExtensionInstance> ExtensionInstance : ExtensionInstances)
 			{
diff --git a/Engine/Plugins/Developer/ProjectLauncher/Source/ProjectLauncher/Private/Widgets/Shared/SCustomLaunchCultureListView.cpp b/Engine/Plugins/Developer/ProjectLauncher/Source/ProjectLauncher/Private/Widgets/Shared/SCustomLaunchCultureListView.cpp
new file mode 100644
--- /dev/null
+++ b/Engine/Plugins/Developer/ProjectLauncher/Source/ProjectLauncher/Private/Widgets/Shared/SCustomLaunchCultureListView.cpp
@@ -0,0 +1,337 @@
+// Copyright Epic Games, Inc. All Rights Reserved.
+
+#include "Widgets/Shared/SCustomLaunchCultureListView.h"
+
+#include "Widgets/SBoxPanel.h"
+#include "Widgets/Views/STreeView.h"
+#include "Widgets/Input/SCheckBox.h"
+#include "Widgets/Input/SComboButton.h"
+#include "Widgets/Input/SSearchBox.h"
+#include "Widgets/Text/STextBlock.h"
+#include "Widgets/Images/SImage.h"
+#include "Widgets/Views/STableRow.h"
+#include "SlateOptMacros.h"
+#include "Styling/AppStyle.h"
+#include "Styling/ProjectLauncherStyle.h"
+#include "Framework/MultiBox/MultiBoxBuilder.h"
+#include "Internationalization/Culture.h"
+#include "Internationalization/Internationalization.h"
+
+
+#define LOCTEXT_NAMESPACE "SCustomLaunchCultureListView"
+
+
+namespace
+{
+	static const TArray<FString> Group_Empty;
+	static const TArray<FString> Group_EFIGS   { TEXT("en"), TEXT("fr"), TEXT("it"), TEXT("de"), TEXT("es") };
+	static const TArray<FString> Group_CJK     { TEXT("zh"), TEXT("ja"), TEXT("ko") };
+	static const TArray<FString> Group_Nordics { TEXT("sv"), TEXT("nb"), TEXT("da"), TEXT("fi"), TEXT("is") };
+}
+
+BEGIN_SLATE_FUNCTION_BUILD_OPTIMIZATION
+void SCustomLaunchCultureListView::Construct(const FArguments& InArgs, TSharedRef<ProjectLauncher::FModel> InModel)
+{
+	Model = InModel.ToSharedPtr();
+	OnSelectionChanged = InArgs._OnSelectionChanged;
+	SelectedCultures = InArgs._SelectedCultures;
+	ProjectPath = InArgs._ProjectPath;
+
+	// start out only showing the current selection if there is one, for easier at-a-glance verification
+	LocalizationGroupFilter = SelectedCultures.Get().Num() == 0 ? ELocalizationGroupFilter::All : ELocalizationGroupFilter::Selected;
+
+    LocalizationGroupDetails.Emplace(ELocalizationGroupFilter::All,			FLocalizationGroup(LOCTEXT("LocGroupName_All",     "All"),					Group_Empty		));
+    LocalizationGroupDetails.Emplace(ELocalizationGroupFilter::Selected,	FLocalizationGroup(LOCTEXT("LocGroupName_Selected","Selected Cultures"),	Group_Empty		));
+    LocalizationGroupDetails.Emplace(ELocalizationGroupFilter::EFIGS,		FLocalizationGroup(LOCTEXT("LocGroupName_EFIGS",   "EFIGS"),				Group_EFIGS		));
+    LocalizationGroupDetails.Emplace(ELocalizationGroupFilter::CJK,			FLocalizationGroup(LOCTEXT("LocGroupName_CJK",     "CJK"),					Group_CJK		));
+    LocalizationGroupDetails.Emplace(ELocalizationGroupFilter::Nordics,		FLocalizationGroup(LOCTEXT("LocGroupName_Nordics", "Nordics"),				Group_Nordics	));
+	check(LocalizationGroupDetails.Num() == (int32)ELocalizationGroupFilter::MAX);
+
+	ChildSlot
+	[
+		SNew(SBorder)
+		.Padding(2)
+		.BorderImage(FAppStyle::GetBrush("Brushes.Background"))
+		[
+			SAssignNew(CultureListView, STreeView<FCultureTreeNodePtr>)
+			.TreeItemsSource(&CultureListViewItemsSource)
+			.OnGenerateRow(this, &SCustomLaunchCultureListView::GenerateCultureTreeNodeRow)
+			.OnGetChildren(this, &SCustomLaunchCultureListView::GetCultureTreeNodeChildren)
+		]
+	];
+
+	RefreshCultureList();
+}
+END_SLATE_FUNCTION_BUILD_OPTIMIZATION
+
+
+
+
+
+
+BEGIN_SLATE_FUNCTION_BUILD_OPTIMIZATION
+TSharedRef<SWidget> SCustomLaunchCultureListView::MakeControlsWidget()
+{
+	FMenuBuilder LocalizationGroupMenuBuilder(true, nullptr);
+	for (int32 GroupIndex = 0; GroupIndex < (int32)ELocalizationGroupFilter::MAX; GroupIndex++)
+	{
+		ELocalizationGroupFilter Group = (ELocalizationGroupFilter)GroupIndex;
+
+		FUIAction Action(
+			FExecuteAction::CreateLambda( [this,Group]()
+			{
+				LocalizationGroupFilter = Group;
+				RefreshCultureList();
+			}),
+			FCanExecuteAction(),
+			FGetActionCheckState::CreateLambda( [this,Group]()
+			{
+				return (LocalizationGroupFilter == Group) ? ECheckBoxState::Checked : ECheckBoxState::Unchecked;
+			})
+		);
+
+		LocalizationGroupMenuBuilder.AddMenuEntry(LocalizationGroupDetails[Group].DisplayName, LOCTEXT("LocGroup_Tip", "Show only cultures in this group"), FSlateIcon(), Action, NAME_None, EUserInterfaceActionType::RadioButton);
+	}
+
+
+	return SNew(SHorizontalBox)
+		+SHorizontalBox::Slot()
+		.AutoWidth()
+		[
+			SNew(SComboButton)
+			.ButtonContent()
+			[
+				SNew(STextBlock)
+				.Text_Lambda( [this]() { return LocalizationGroupDetails[LocalizationGroupFilter].DisplayName; } )
+			]
+			.MenuContent()
+			[
+				LocalizationGroupMenuBuilder.MakeWidget()
+			]
+		]
+
+
+		+SHorizontalBox::Slot()
+		.FillWidth(1)
+		.Padding(4,0)
+		[
+			SNew(SSearchBox)
+			.OnTextCommitted(this, &SCustomLaunchCultureListView::OnSearchFilterTextCommitted)
+			.OnTextChanged(this, &SCustomLaunchCultureListView::OnSearchFilterTextChanged)
+		]
+	;
+}
+END_SLATE_FUNCTION_BUILD_OPTIMIZATION
+
+
+void SCustomLaunchCultureListView::GetCultureTreeNodeChildren(FCultureTreeNodePtr Node, TArray<FCultureTreeNodePtr>& OutChildren)
+{
+	OutChildren = Node->Children;
+}
+
+
+
+BEGIN_SLATE_FUNCTION_BUILD_OPTIMIZATION
+TSharedRef<ITableRow> SCustomLaunchCultureListView::GenerateCultureTreeNodeRow(FCultureTreeNodePtr Node, const TSharedRef<STableViewBase>& OwnerTable)
+{
+	return SNew(STableRow<FCultureTreeNodePtr>, OwnerTable)
+	.ToolTipText(this, &SCustomLaunchCultureListView::GetCultureDisplayName, Node)
+	[
+		SNew(SCheckBox)
+		.Padding(2)
+		.IsChecked(this, &SCustomLaunchCultureListView::GetCultureTreeNodeCheckState, Node)
+		.OnCheckStateChanged(this, &SCustomLaunchCultureListView::SetCultureTreeNodeCheckState, Node)
+		[
+			SNew(STextBlock)
+			.Text(FText::FromString(*Node->Name))
+		]
+	];
+}
+END_SLATE_FUNCTION_BUILD_OPTIMIZATION
+
+
+
+FText SCustomLaunchCultureListView::GetCultureDisplayName(FCultureTreeNodePtr Item) const
+{
+	const FCulturePtr CulturePtr = FInternationalization::Get().GetCulture(Item->Name);
+	if (CulturePtr == nullptr)
+	{
+		return FText::GetEmpty();
+	}
+
+	return FText::FromString(CulturePtr->GetDisplayName());
+}
+
+ECheckBoxState SCustomLaunchCultureListView::GetCultureTreeNodeCheckState(FCultureTreeNodePtr Node) const
+{
+	return SelectedCultures.Get().Contains(Node->Name) ? ECheckBoxState::Checked : ECheckBoxState::Unchecked;
+}
+
+
+
+void SCustomLaunchCultureListView::SetCultureTreeNodeCheckState(ECheckBoxState CheckBoxState, FCultureTreeNodePtr Node)
+{
+	TArray<FString> CheckedCultures = SelectedCultures.Get();
+
+	if (CheckBoxState == ECheckBoxState::Checked)
+	{
+		CheckedCultures.Add(Node->Name);
+	}
+	else
+	{
+		CheckedCultures.Remove(Node->Name);
+	}
+
+	OnSelectionChanged.ExecuteIfBound(CheckedCultures);
+}
+
+
+
+void SCustomLaunchCultureListView::OnSearchFilterTextCommitted(const FText& SearchText, ETextCommit::Type CommitType)
+{
+	CurrentFilterText = SearchText.ToString();
+	RefreshCultureList();
+}
+
+
+
+void SCustomLaunchCultureListView::OnSearchFilterTextChanged(const FText& SearchText)
+{
+	CurrentFilterText = SearchText.ToString();
+	if (CurrentFilterText.IsEmpty())
+	{
+		RefreshCultureList();
+	}
+}
+
+
+
+void SCustomLaunchCultureListView::OnProjectChanged()
+{
+	RefreshCultureList();
+}
+
+
+
+void SCustomLaunchCultureListView::RefreshCultureList()
+{
+	bCultureListDirty = true;
+}
+
+void SCustomLaunchCultureListView::RefreshCultureListInternal()
+{
+	bCultureListDirty = false;
+
+	CultureTreeRoot = MakeShared<FCultureTreeNode>();
+
+	TArray<FString> SelectedCultureNames = SelectedCultures.Get();
+
+	TArray<FString> CultureNames;
+	FInternationalization::Get().GetCultureNames(CultureNames);
+	CultureNames.Sort();
+
+	for (const FString& Culture : CultureNames)
+	{
+		FString Language, Region;
+		if (!Culture.Split(TEXT("-"), &Language, &Region))
+		{
+			Language = Culture;
+		}
+
+		// ignore this culture if it isn't filtered
+		if (!CurrentFilterText.IsEmpty())
+		{
+			const FCulturePtr CulturePtr = FInternationalization::Get().GetCulture(Culture);
+			if (!Culture.Contains(CurrentFilterText) && (CulturePtr == nullptr || !CulturePtr->GetDisplayName().Contains(CurrentFilterText)))
+			{
+				continue;
+			}
+		}
+
+		// check the localization group filter table
+		if (LocalizationGroupFilter == ELocalizationGroupFilter::Selected)
+		{
+			if (!SelectedCultureNames.Contains(Culture))
+			{
+				continue;
+			}
+		}
+		else if (LocalizationGroupFilter != ELocalizationGroupFilter::All)
+		{
+			if (!LocalizationGroupDetails[LocalizationGroupFilter].Languages.Contains(Language))
+			{
+				continue;
+			}
+		}
+
+
+		// add the culture, splitting into Language/Region hierarchy as necessary
+		if (Region.IsEmpty())
+		{
+			FCultureTreeNodePtr LanguageNode = CultureTreeRoot->Children.Add_GetRef(MakeShared<FCultureTreeNode>());
+			LanguageNode->Name = Culture;
+		}
+		else
+		{
+			// add or find the language
+			FCultureTreeNodePtr LanguageNode = nullptr;
+			if (FCultureTreeNodePtr* ChildPtr = CultureTreeRoot->Children.FindByPredicate( [Language]( const FCultureTreeNodePtr Node ) { return Node->Name == Language; } ))
+			{
+				LanguageNode = *ChildPtr;
+			}
+			else
+			{
+				LanguageNode = CultureTreeRoot->Children.Add_GetRef(MakeShared<FCultureTreeNode>());
+				LanguageNode->Name = Language;
+			}
+
+			// add the region/country as a child
+			FCultureTreeNodePtr RegionNode = LanguageNode->Children.Add_GetRef(MakeShared<FCultureTreeNode>());
+			RegionNode->Name = Culture;
+		}
+	}
+
+
+	CultureListViewItemsSource = CultureTreeRoot->Children;
+	if (CultureListView.IsValid())
+	{
+		CultureListView->RequestTreeRefresh();
+	}
+
+	// expand the parents of anything that is selected
+	for (FCultureTreeNodePtr LanguageNode : CultureTreeRoot->Children)
+	{
+		for (FCultureTreeNodePtr RegionNode : LanguageNode->Children)
+		{
+			if (SelectedCultureNames.Contains(RegionNode->Name))
+			{
+				CultureListView->SetItemExpansion(LanguageNode, true);
+				break;
+			}
+		}
+	}
+}
+
+
+
+void SCustomLaunchCultureListView::Tick(const FGeometry& AllottedGeometry, const double InCurrentTime, const float InDeltaTime)
+{
+	if (bCultureListDirty && bHasPaintedThisFrame)
+	{
+		RefreshCultureListInternal();
+	}
+
+	bHasPaintedThisFrame = false;
+	SCompoundWidget::Tick(AllottedGeometry, InCurrentTime, InDeltaTime);
+}
+
+int32 SCustomLaunchCultureListView::OnPaint( const FPaintArgs& Args, const FGeometry& AllottedGeometry, const FSlateRect& MyCullingRect, FSlateWindowElementList& OutDrawElements, int32 LayerId, const FWidgetStyle& InWidgetStyle, bool bParentEnabled) const
+{
+	bHasPaintedThisFrame = true;
+    return SCompoundWidget::OnPaint(Args, AllottedGeometry, MyCullingRect, OutDrawElements, LayerId, InWidgetStyle, bParentEnabled);
+}
+
+
+
+
+#undef LOCTEXT_NAMESPACE
diff --git a/Engine/Plugins/Developer/ProjectLauncher/Source/ProjectLauncher/Public/Extension/LaunchExtension.h b/Engine/Plugins/Developer/ProjectLauncher/Source/ProjectLauncher/Public/Extension/LaunchExtension.h
--- a/Engine/Plugins/Developer/ProjectLauncher/Source/ProjectLauncher/Public/Extension/LaunchExtension.h
+++ b/Engine/Plugins/Developer/ProjectLauncher/Source/ProjectLauncher/Public/Extension/LaunchExtension.h
@@ -33,8 +33,10 @@ namespace ProjectLauncher
 		struct FArgs
 		{
 			ILauncherProfileRef Profile;
+			ILauncherProfileRef DefaultProfile;
 			TSharedRef<FModel> Model;
 			TSharedRef<FLaunchExtension> Extension;
+			FLaunchProfileTreeData& OwnerTreeData;
 		};
 
 		UE_API FLaunchExtensionInstance( FArgs& InArgs );
@@ -86,6 +88,11 @@ namespace ProjectLauncher
 		 */
 		virtual void OnPropertyChanged() {};
 
+		/**
+		 * Notification callback when the project has been changed in our profile
+		 */
+		virtual void OnProjectChanged() {};
+
 		/**
 		 * Notification callback when the profile is validated
 		 */
@@ -219,6 +226,11 @@ namespace ProjectLauncher
 		 */
 		inline ILauncherProfileRef GetProfile() const { return Profile; }
 
+		/** 
+		 * Get the reference default profile
+		 */
+		inline ILauncherProfileRef GetDefaultProfile() const { return DefaultProfile; }
+
 		/** 
 		 * Get the model, for general purpose helper functions
 		 */
@@ -332,14 +344,22 @@ namespace ProjectLauncher
 		 */
 		UE_API void BroadcastPropertyChanged() const;
 
+
+		/** 
+		 * Request that the property tree should be refreshed - such as when a widget has changed size
+		*/
+		UE_API void RequestTreeRefresh();
+
 	private:
 		UE_API bool IsParameterGroup( const FString& InParameter ) const;
 		UE_API bool TryRemoveParameterGroup( const FString& InParameter );
 
 
 		ILauncherProfileRef Profile;
+		ILauncherProfileRef DefaultProfile;
 		TSharedRef<FModel> Model;
 		TSharedRef<FLaunchExtension> Extension;
+		FLaunchProfileTreeData& OwnerTreeData;
 		FPropertyChangedDelegate PropertyChangedDelegate;
 
 		friend class FLaunchProfileTreeNode;
@@ -437,9 +457,10 @@ namespace ProjectLauncher
 		 * 
 		 * @param InProfile the current profile 
 		 * @param InModel helper class
+		 * @param InOwnerTreeData the property tree data associated with the profile
 		 * @returns array of all compatible extensions
 		 */
-		UE_API static TArray<TSharedPtr<FLaunchExtensionInstance>> CreateExtensionInstancesForProfile( ILauncherProfileRef InProfile, TSharedRef<FModel> InModel );
+		UE_API static TArray<TSharedPtr<FLaunchExtensionInstance>> CreateExtensionInstancesForProfile( ILauncherProfileRef InProfile, TSharedRef<FModel> InModel, FLaunchProfileTreeData& InOwnerTreeData );
 
 
 	private:
diff --git a/Engine/Plugins/Developer/ProjectLauncher/Source/ProjectLauncher/Public/Model/ProjectLauncherModel.h b/Engine/Plugins/Developer/ProjectLauncher/Source/ProjectLauncher/Public/Model/ProjectLauncherModel.h
--- a/Engine/Plugins/Developer/ProjectLauncher/Source/ProjectLauncher/Public/Model/ProjectLauncherModel.h
+++ b/Engine/Plugins/Developer/ProjectLauncher/Source/ProjectLauncher/Public/Model/ProjectLauncherModel.h
@@ -108,12 +108,6 @@ namespace ProjectLauncher
 	/* Delegate fired when the selected profile changes */
 	DECLARE_MULTICAST_DELEGATE_TwoParams(FOnSelectedProfileChanged, const ILauncherProfilePtr& /*new*/, const ILauncherProfilePtr& /*old*/)
 
-	/* Delegate fired when the selected project changes */
-	DECLARE_MULTICAST_DELEGATE(FOnSelectedProjectChanged)
-
-	/* Delegate fired when the active build target changes */
-	DECLARE_MULTICAST_DELEGATE(FOnBuildTargetChanged)
-
 	/* Delegate fired when the selected platform changes */
 	DECLARE_MULTICAST_DELEGATE(FOnSelectedPlatformChanged)
 
@@ -252,6 +246,7 @@ namespace ProjectLauncher
 		/* Retrieve project settings from a project path or profile */
 		UE_API const FProjectSettings GetProjectSettings(const FString& ProjectPath);
 		UE_API const FProjectSettings GetProjectSettings(const ILauncherProfileRef& Profile);
+		UE_API FString GetProjectPath(const ILauncherProfileRef& Profile) const;
 
 		/* Get config section and ini file used for storing project launcher state */
 		UE_API const TCHAR* GetConfigSection() const;
diff --git a/Engine/Plugins/Developer/ProjectLauncher/Source/ProjectLauncher/Public/ProfileTree/GenericProfileTreeBuilder.h b/Engine/Plugins/Developer/ProjectLauncher/Source/ProjectLauncher/Public/ProfileTree/GenericProfileTreeBuilder.h
--- a/Engine/Plugins/Developer/ProjectLauncher/Source/ProjectLauncher/Public/ProfileTree/GenericProfileTreeBuilder.h
+++ b/Engine/Plugins/Developer/ProjectLauncher/Source/ProjectLauncher/Public/ProfileTree/GenericProfileTreeBuilder.h
@@ -8,7 +8,6 @@
 #define UE_API PROJECTLAUNCHER_API
 
 class SSearchableComboBox;
-class SCustomLaunchMapListView;
 class SCustomLaunchDeviceListView;
 namespace ESelectInfo { enum Type : int; }
 
@@ -56,10 +55,7 @@ namespace ProjectLauncher
 		UE_API void AddImportZenSnapshotProperty( FLaunchProfileTreeNode& HeadingNode );
 		UE_API void AddZenPakStreamingPathProperty( FLaunchProfileTreeNode& HeadingNode );
 		UE_API void AddIncrementalCookProperty( FLaunchProfileTreeNode& HeadingNode );
-		UE_API void AddUnversionedCookProperty( FLaunchProfileTreeNode& HeadingNode );
 		UE_API void AddCookProperty( FLaunchProfileTreeNode& HeadingNode );
-		UE_API void AddMapsToCookProperty( FLaunchProfileTreeNode& HeadingNode );
-		UE_API void AddAdditionalCookerOptionsProperty( FLaunchProfileTreeNode& HeadingNode );
 		UE_API void AddBuildProperty( FLaunchProfileTreeNode& HeadingNode );
 		UE_API void AddForceBuildProperty( FLaunchProfileTreeNode& HeadingNode );
 		UE_API void AddBuildUATProperty( FLaunchProfileTreeNode& HeadingNode );
@@ -88,13 +84,6 @@ namespace ProjectLauncher
 
 
 
-		TArray<FString> CachedMapsToCook;
-		enum class EMapOption : uint8
-		{
-			Startup,
-			Selected,
-		};
-		EMapOption MapOption = EMapOption::Startup;
 
 		TArray<FString> CachedDeployDeviceIDs;
 		enum class EDeployDeviceOption : uint8
@@ -126,16 +115,6 @@ namespace ProjectLauncher
 		UE_API void SetCook( bool bCook );
 		UE_API bool GetCook( ILauncherProfilePtr InProfile = nullptr ) const;
 
-		UE_API void SetMapsToCook(TArray<FString> MapsToCook);
-		UE_API TArray<FString> GetMapsToCook() const;
-		UE_API EMapOption GetMapOption() const;
-		UE_API void SetMapOption( EMapOption MapOption );
-		UE_API float GetMapListHeight() const;
-		UE_API void SetMapListHeight( float NewHeight );
-		UE_API TSharedRef<SWidget> CreateMapListWidget();
-		float MapListHeight = 300.0f;
-
-
 		UE_API void SetDeployDeviceIDs(TArray<FString> DeployDeviceIDs);
 		UE_API TArray<FString> GetDeployDeviceIDs() const;
 		UE_API EDeployDeviceOption GetDeployDeviceOption() const;
@@ -189,7 +168,6 @@ namespace ProjectLauncher
 		mutable bool bStartupMapCacheDirty = true;
 		mutable TArray<TSharedPtr<FString>> CachedStartupMaps;
 		TSharedPtr<SSearchableComboBox> InitalMapCombo;
-		TSharedPtr<SCustomLaunchMapListView> MapListView;
 		TSharedPtr<SCustomLaunchDeviceListView> DeployDeviceListView;
 		TArray<FString> CachedArchitectures;
 		EBuildTargetType CachedBuildTargetType = EBuildTargetType::Game;
diff --git a/Engine/Plugins/Developer/ProjectLauncher/Source/ProjectLauncher/Public/Widgets/Shared/SCustomLaunchCultureListView.h b/Engine/Plugins/Developer/ProjectLauncher/Source/ProjectLauncher/Public/Widgets/Shared/SCustomLaunchCultureListView.h
new file mode 100644
--- /dev/null
+++ b/Engine/Plugins/Developer/ProjectLauncher/Source/ProjectLauncher/Public/Widgets/Shared/SCustomLaunchCultureListView.h
@@ -0,0 +1,106 @@
+// Copyright Epic Games, Inc. All Rights Reserved.
+
+#pragma once
+
+#include "CoreMinimal.h"
+#include "Widgets/SCompoundWidget.h"
+#include "Model/ProjectLauncherModel.h"
+#include <initializer_list>
+
+#define UE_API PROJECTLAUNCHER_API
+
+class ITableRow;
+class STableViewBase;
+template<typename ItemType> class STreeView;
+
+class SCustomLaunchCultureListView
+	: public SCompoundWidget
+{
+public:
+	DECLARE_DELEGATE_OneParam(FOnSelectionChanged, TArray<FString> );
+
+	SLATE_BEGIN_ARGS(SCustomLaunchCultureListView){}
+		SLATE_EVENT(FOnSelectionChanged, OnSelectionChanged);
+		SLATE_ATTRIBUTE(TArray<FString>, SelectedCultures);
+		SLATE_ATTRIBUTE(FString, ProjectPath)
+	SLATE_END_ARGS()
+
+public:
+	UE_API void Construct( const FArguments& InArgs, TSharedRef<ProjectLauncher::FModel> InModel );
+
+	UE_API void RefreshCultureList();
+	UE_API TSharedRef<SWidget> MakeControlsWidget();
+
+protected:
+	TSharedPtr<ProjectLauncher::FModel> Model;
+	TAttribute<TArray<FString>> SelectedCultures;
+	TAttribute<FString> ProjectPath;
+	FOnSelectionChanged OnSelectionChanged;
+
+private:
+
+	struct FCultureTreeNode
+	{
+		FString Name;
+		TArray<TSharedPtr<FCultureTreeNode>> Children;
+	};
+	typedef TSharedPtr<FCultureTreeNode> FCultureTreeNodePtr;
+
+	FCultureTreeNodePtr CultureTreeRoot;
+
+	UE_API void OnProjectChanged();
+
+	UE_API void OnSearchFilterTextCommitted(const FText& SearchText, ETextCommit::Type InCommitType);
+	UE_API void OnSearchFilterTextChanged(const FText& SearchText);
+
+	FString CurrentFilterText;
+
+	enum class ELocalizationGroupFilter : int32
+	{
+		All,
+		Selected,
+		EFIGS,
+		CJK,
+		Nordics,
+
+		MAX
+	};
+	ELocalizationGroupFilter LocalizationGroupFilter;
+
+	struct FLocalizationGroup
+	{
+		FLocalizationGroup( FText InDisplayName, const TArray<FString>& InLanguages )
+			: DisplayName(InDisplayName)
+			, Languages(InLanguages)
+		{
+		}
+
+		FText DisplayName;
+		const TArray<FString>& Languages;
+	};
+	TMap<ELocalizationGroupFilter,FLocalizationGroup> LocalizationGroupDetails;
+
+
+
+	static const TCHAR** LocalizationGroupPrefixes[(int32)ELocalizationGroupFilter::MAX];
+	static const FText LocalizationGroupDisplayName[(int32)ELocalizationGroupFilter::MAX];
+
+
+	TSharedPtr<STreeView<FCultureTreeNodePtr>> CultureListView;
+	TArray<FCultureTreeNodePtr> CultureListViewItemsSource;
+
+	UE_API void GetCultureTreeNodeChildren(FCultureTreeNodePtr Item, TArray<FCultureTreeNodePtr>& OutChildren);
+	UE_API TSharedRef<ITableRow> GenerateCultureTreeNodeRow(FCultureTreeNodePtr Item, const TSharedRef<STableViewBase>& OwnerTable);
+	UE_API ECheckBoxState GetCultureTreeNodeCheckState(FCultureTreeNodePtr Item) const;
+	UE_API void SetCultureTreeNodeCheckState(ECheckBoxState CheckBoxState, FCultureTreeNodePtr Item);
+	UE_API FText GetCultureDisplayName(FCultureTreeNodePtr Item) const;
+
+	mutable bool bHasPaintedThisFrame = false;
+	UE_API virtual int32 OnPaint( const FPaintArgs& Args, const FGeometry& AllottedGeometry, const FSlateRect& MyCullingRect, FSlateWindowElementList& OutDrawElements, int32 LayerId, const FWidgetStyle& InWidgetStyle, bool bParentEnabled) const override;
+
+	bool bCultureListDirty = false;
+	UE_API void RefreshCultureListInternal();
+	UE_API virtual void Tick(const FGeometry& AllottedGeometry, const double InCurrentTime, const float InDeltaTime) override;
+};
+
+#undef UE_API
