ensure FDialogCommands are initialized for Project Launcher

#rnx
#Synced-CL 47728819

diff --git a/Engine/Plugins/Developer/ProjectLauncher/Source/ProjectLauncher/Private/ProjectLauncherModule.cpp b/Engine/Plugins/Developer/ProjectLauncher/Source/ProjectLauncher/Private/ProjectLauncherModule.cpp
--- a/Engine/Plugins/Developer/ProjectLauncher/Source/ProjectLauncher/Private/ProjectLauncherModule.cpp
+++ b/Engine/Plugins/Developer/ProjectLauncher/Source/ProjectLauncher/Private/ProjectLauncherModule.cpp
@@ -3,6 +3,7 @@
 #include "ProjectLauncherModule.h"
 #include "WorkspaceMenuStructure.h"
 #include "WorkspaceMenuStructureModule.h"
+#include "Dialog/DialogCommands.h"
 #include "Framework/Docking/TabManager.h"
 #include "Widgets/Docking/SDockTab.h"
 #include "Widgets/SProjectLauncher.h"
@@ -38,8 +39,12 @@ public:
 		RegisterTreeBuilder( MakeShared<ProjectLauncher::FBasicProfileTreeBuilderFactory>() );
 		RegisterTreeBuilder( MakeShared<ProjectLauncher::FCustomProfileTreeBuilderFactory>() );
 
-		// register custom styles
+		// register custom styles & other UI elements
 		FProjectLauncherStyle::Initialize();
+		if (!FDialogCommands::IsRegistered())
+		{
+			FDialogCommands::Register();
+		}
 
 		// register main tab
 		FGlobalTabmanager::Get()->RegisterNomadTabSpawner(TabName, FOnSpawnTab::CreateRaw(this, &FProjectLauncherModule::HandleSpawnTab))
