Add support for full-row project launcher property rows - expected for banners etc.
- set the widget name to 'empty'
- no 'reset to default' support at the moment

#rnx
#Synced-CL 47728819

diff --git a/Engine/Plugins/Developer/ProjectLauncher/Source/ProjectLauncher/Private/Widgets/Profiles/SCustomLaunchCustomProfileEditor.cpp b/Engine/Plugins/Developer/ProjectLauncher/Source/ProjectLauncher/Private/Widgets/Profiles/SCustomLaunchCustomProfileEditor.cpp
--- a/Engine/Plugins/Developer/ProjectLauncher/Source/ProjectLauncher/Private/Widgets/Profiles/SCustomLaunchCustomProfileEditor.cpp
+++ b/Engine/Plugins/Developer/ProjectLauncher/Source/ProjectLauncher/Private/Widgets/Profiles/SCustomLaunchCustomProfileEditor.cpp
@@ -160,6 +160,43 @@ public:
 	}
 };
 
+class SLaunchProfileFullWidthTreeRow : public SLaunchProfileTreeRow
+{
+public:
+	SLATE_BEGIN_ARGS(SLaunchProfileFullWidthTreeRow){}
+	SLATE_END_ARGS()
+
+	void Construct(const FArguments& InArgs, ProjectLauncher::FLaunchProfileTreeNodePtr TreeNode, const TSharedRef<STableViewBase>& OwnerTable )
+	{
+		this->ChildSlot
+		[
+			SNew(SBorder)
+			.BorderImage(FAppStyle::Get().GetBrush("Brushes.Recessed"))
+			.Padding(FMargin(0, 0, 1, 0))
+			[
+				SNew(SBorder)
+				.BorderImage_Lambda( [this]() { return IsHovered() ? FAppStyle::Get().GetBrush("Brushes.Header") : FAppStyle::Get().GetBrush("Brushes.Panel"); } )
+				.Padding(1)
+				[
+					SNew(SBox)
+					.Padding( FMargin(24, 1, 24, 1))
+					.VAlign(VAlign_Center)
+					[
+						TreeNode->Widget.ToSharedRef()
+					]
+				]
+			]
+		];
+
+		SLaunchProfileTreeRow::ConstructInternal(
+			STableRow::FArguments()
+			.ShowSelection(false),
+			OwnerTable
+		);
+	}
+};
+
+
 class SLaunchProfilePropertyTreeRow : public SLaunchProfileTreeRow
 {
 public:
@@ -452,18 +489,25 @@ END_SLATE_FUNCTION_BUILD_OPTIMIZATION
 
 TSharedRef<ITableRow> SCustomLaunchCustomProfileEditor::OnGenerateWidgetForTreeNode( ProjectLauncher::FLaunchProfileTreeNodePtr TreeNode, const TSharedRef<STableViewBase>& OwnerTable )
 {
-	if (TreeNode->Widget.IsValid())
+	auto IsVisible = [TreeNode]()
 	{
-		auto IsVisible = [TreeNode]()
+		if (TreeNode->Callbacks.IsVisible == nullptr)
 		{
-			if (TreeNode->Callbacks.IsVisible == nullptr)
-			{
-				return true;
-			}
+			return true;
+		}
 
-			return TreeNode->Callbacks.IsVisible();
-		};
+		return TreeNode->Callbacks.IsVisible();
+	};
 
+	if (TreeNode->Name.IsEmpty() && TreeNode->Widget.IsValid())
+	{
+		return SNew(SLaunchProfileFullWidthTreeRow, TreeNode, OwnerTable)
+		.Visibility_Lambda( [IsVisible]() { return IsVisible() ? EVisibility::Visible : EVisibility::Collapsed; } )
+		.IsEnabled_Lambda( [TreeNode]() { return TreeNode->Callbacks.IsEnabled ? TreeNode->Callbacks.IsEnabled() : true; } )
+		;
+	}
+	else if (TreeNode->Widget.IsValid())
+	{
 		return SNew(SLaunchProfilePropertyTreeRow, TreeNode, OwnerTable)
 		.SplitterValue_Lambda( [this]() { return SplitterPos; } )
 		.OnSplitterResized_Lambda( [this](float NewPos) { SplitterPos = FMath::Clamp(NewPos, 0.0f, 1.0f); } )
