Two fixes for Project Launcher extensions:
- config values are not cleared when setting 0, false or empty string - this might not be the default value
- launch extension registry is cleared when project launcher shuts down to avoid waiting until atexit handlers.

#rnx
#Synced-CL 48214739

diff --git a/Engine/Plugins/Developer/ProjectLauncher/Source/ProjectLauncher/Private/Extension/LaunchExtension.cpp b/Engine/Plugins/Developer/ProjectLauncher/Source/ProjectLauncher/Private/Extension/LaunchExtension.cpp
--- a/Engine/Plugins/Developer/ProjectLauncher/Source/ProjectLauncher/Private/Extension/LaunchExtension.cpp
+++ b/Engine/Plugins/Developer/ProjectLauncher/Source/ProjectLauncher/Private/Extension/LaunchExtension.cpp
@@ -11,6 +11,7 @@
 #include "Modules/ModuleManager.h"
 #include "ILauncherServicesModule.h"
 #include "PlatformInfo.h"
+#include <atomic>
 
 #define LOCTEXT_NAMESPACE "FLaunchExtensionInstance"
 
@@ -20,18 +21,34 @@ namespace ProjectLauncher
 	static TArray<TSharedPtr<FLaunchExtension>> GExtensions;
 	static TMap<ILauncherProfilePtr, TArray<TSharedPtr<FLaunchExtensionInstance>>> GExtensionInstances;
 	static FCriticalSection GExtensionsCS;
+	static std::atomic<bool> GExtensionsDestroyed = false;
+
+
+	void ShutdownLaunchExtensions()
+	{
+		check(!GExtensionsDestroyed);
+		GExtensions.Reset();
+		GExtensionInstances.Reset();
+
+		GExtensionsDestroyed = true;
+	}
 
 
 	void RegisterExtension( TSharedRef<FLaunchExtension> Extension )
 	{
+		check(!GExtensionsDestroyed);
+
 		FScopeLock Lock(&GExtensionsCS);
 		GExtensions.Add(Extension);
 	}
 
 	void UnregisterExtension( TSharedRef<FLaunchExtension> Extension )
 	{
-		FScopeLock Lock(&GExtensionsCS);
-		GExtensions.Remove(Extension);
+		if (!GExtensionsDestroyed)
+		{
+			FScopeLock Lock(&GExtensionsCS);
+			GExtensions.Remove(Extension);
+		}
 	}
 
 
@@ -127,6 +144,8 @@ namespace ProjectLauncher
 
 	TArray<TSharedPtr<FLaunchExtensionInstance>> FLaunchExtension::CreateExtensionInstancesForProfile( ILauncherProfileRef InProfile, TSharedRef<FModel> InModel, FLaunchProfileTreeData& OwnerTreeData )
 	{
+		check(!GExtensionsDestroyed);
+
 		// make sure the extensions for this profile haven't been created already
 		{
 			FScopeLock Lock(&GExtensionsCS);
@@ -533,14 +552,7 @@ namespace ProjectLauncher
 
 		if (Config == EConfig::PerProfile)
 		{
-			if (Value.IsEmpty())
-			{
-				Profile->GetCustomStringProperties().Remove(KeyName);
-			}
-			else
-			{
-				Profile->GetCustomStringProperties().Add(KeyName, Value);
-			}
+			Profile->GetCustomStringProperties().Add(KeyName, Value);
 			BroadcastPropertyChanged();
 		}
 		else
@@ -555,14 +567,7 @@ namespace ProjectLauncher
 
 		if (Config == EConfig::PerProfile)
 		{
-			if (Value == false)
-			{
-				Profile->GetCustomBoolProperties().Remove(KeyName);
-			}
-			else
-			{
-				Profile->GetCustomBoolProperties().Add(KeyName, Value);
-			}
+			Profile->GetCustomBoolProperties().Add(KeyName, Value);
 			BroadcastPropertyChanged();
 		}
 		else
@@ -577,14 +582,7 @@ namespace ProjectLauncher
 
 		if (Config == EConfig::PerProfile)
 		{
-			if (Value == 0)
-			{
-				Profile->GetCustomStringProperties().Remove(KeyName);
-			}
-			else
-			{
-				Profile->GetCustomStringProperties().Add(KeyName, ::LexToString(Value) );
-			}
+			Profile->GetCustomStringProperties().Add(KeyName, ::LexToString(Value) );
 			BroadcastPropertyChanged();
 		}
 		else
@@ -599,14 +597,7 @@ namespace ProjectLauncher
 
 		if (Config == EConfig::PerProfile)
 		{
-			if (Value == 0.0f)
-			{
-				Profile->GetCustomStringProperties().Remove(KeyName);
-			}
-			else
-			{
-				Profile->GetCustomStringProperties().Add(KeyName, ::LexToString(Value) );
-			}
+			Profile->GetCustomStringProperties().Add(KeyName, ::LexToString(Value) );
 			BroadcastPropertyChanged();
 		}
 		else
diff --git a/Engine/Plugins/Developer/ProjectLauncher/Source/ProjectLauncher/Private/Model/ProjectLauncherModel.cpp b/Engine/Plugins/Developer/ProjectLauncher/Source/ProjectLauncher/Private/Model/ProjectLauncherModel.cpp
--- a/Engine/Plugins/Developer/ProjectLauncher/Source/ProjectLauncher/Private/Model/ProjectLauncherModel.cpp
+++ b/Engine/Plugins/Developer/ProjectLauncher/Source/ProjectLauncher/Private/Model/ProjectLauncherModel.cpp
@@ -77,6 +77,8 @@ namespace ProjectLauncher
 
 
 
+	extern void ShutdownLaunchExtensions();
+
 	FModel::~FModel()
 	{
 		for (TPair<FString, FProjectSettings>& Pair : CachedProjectSettings)
@@ -94,6 +96,8 @@ namespace ProjectLauncher
 		DeviceProxyManager->OnProxyRemoved().RemoveAll(this);
 
 		SaveConfig();
+
+		ShutdownLaunchExtensions();
 	}
 
 
