improved lifetime management for Project Launcher Model & extensions

#rnx
#Synced-CL 48214739

diff --git a/Engine/Plugins/Developer/ProjectLauncher/Source/ProjectLauncher/Private/Extension/LaunchExtension.cpp b/Engine/Plugins/Developer/ProjectLauncher/Source/ProjectLauncher/Private/Extension/LaunchExtension.cpp
--- a/Engine/Plugins/Developer/ProjectLauncher/Source/ProjectLauncher/Private/Extension/LaunchExtension.cpp
+++ b/Engine/Plugins/Developer/ProjectLauncher/Source/ProjectLauncher/Private/Extension/LaunchExtension.cpp
@@ -26,11 +26,13 @@ namespace ProjectLauncher
 
 	void ShutdownLaunchExtensions()
 	{
-		check(!GExtensionsDestroyed);
-		GExtensions.Reset();
-		GExtensionInstances.Reset();
+		if (!GExtensionsDestroyed)
+		{
+			GExtensions.Reset();
+			GExtensionInstances.Reset();
 
-		GExtensionsDestroyed = true;
+			GExtensionsDestroyed = true;
+		}
 	}
 
 
diff --git a/Engine/Plugins/Developer/ProjectLauncher/Source/ProjectLauncher/Private/ProjectLauncherModule.cpp b/Engine/Plugins/Developer/ProjectLauncher/Source/ProjectLauncher/Private/ProjectLauncherModule.cpp
--- a/Engine/Plugins/Developer/ProjectLauncher/Source/ProjectLauncher/Private/ProjectLauncherModule.cpp
+++ b/Engine/Plugins/Developer/ProjectLauncher/Source/ProjectLauncher/Private/ProjectLauncherModule.cpp
@@ -26,6 +26,9 @@ namespace ProjectLauncher
 
 	extern void RegisterExtension( TSharedRef<FLaunchExtension> Extension );
 	extern void UnregisterExtension( TSharedRef<FLaunchExtension> Extension );
+
+	extern void ShutdownLaunchExtensions();
+
 }
 
 
@@ -58,6 +61,8 @@ public:
 	{
 		FGlobalTabmanager::Get()->UnregisterNomadTabSpawner(TabName);
 		FProjectLauncherStyle::Shutdown();
+		ProjectLauncher::ShutdownLaunchExtensions();
+		Model.Reset();
 	}
 
 	virtual void RegisterTreeBuilder( TSharedRef<ProjectLauncher::ILaunchProfileTreeBuilderFactory> TreeBuilderFactory ) override
@@ -91,9 +96,12 @@ private:
 		ILauncherServicesModule& LauncherServicesModule = FModuleManager::LoadModuleChecked<ILauncherServicesModule>("LauncherServices");
 		ITargetDeviceServicesModule& TargetDeviceServicesModule = FModuleManager::LoadModuleChecked<ITargetDeviceServicesModule>("TargetDeviceServices");
 
-		TSharedRef<ProjectLauncher::FModel> Model = MakeShared<ProjectLauncher::FModel>(TargetDeviceServicesModule.GetDeviceProxyManager(), LauncherServicesModule.CreateLauncher(), LauncherServicesModule.GetProfileManager() );
+		if (!Model.IsValid())
+		{
+			Model = MakeShared<ProjectLauncher::FModel>(TargetDeviceServicesModule.GetDeviceProxyManager(), LauncherServicesModule.CreateLauncher(), LauncherServicesModule.GetProfileManager() );
+		}
 		
-		auto OnTabActivated = [Model](TSharedRef<SDockTab>, ETabActivationCause Cause)
+		auto OnTabActivated = [this](TSharedRef<SDockTab>, ETabActivationCause Cause)
 		{
 			if (Cause == ETabActivationCause::UserClickedOnTab)
 			{
@@ -110,9 +118,11 @@ private:
 		.TabRole(ETabRole::NomadTab)
 		.OnTabActivated_Lambda(OnTabActivated)
 		[
-			SNew(SProjectLauncher, Model)
+			SNew(SProjectLauncher, Model.ToSharedRef())
 		];
 	}
+
+	TSharedPtr<ProjectLauncher::FModel> Model;
 };
 
 
