New Project Launcher supports multiple build targets / platforms
- only accessible when loading a launcher profile that already contains multiple platforms for now
- improved interface is coming soon

#rnx

#Synced-CL 48214739

diff --git a/Engine/Plugins/Developer/ProjectLauncher/Source/CommonLaunchExtensions/Private/Globals/GlobalsLaunchExtension.cpp b/Engine/Plugins/Developer/ProjectLauncher/Source/CommonLaunchExtensions/Private/Globals/GlobalsLaunchExtension.cpp
--- a/Engine/Plugins/Developer/ProjectLauncher/Source/CommonLaunchExtensions/Private/Globals/GlobalsLaunchExtension.cpp
+++ b/Engine/Plugins/Developer/ProjectLauncher/Source/CommonLaunchExtensions/Private/Globals/GlobalsLaunchExtension.cpp
@@ -60,7 +60,7 @@ bool FGlobalsLaunchExtensionInstance::GetExtensionVariableValue( const FString&
 	}
 	else if (InVariable == TargetNameVariable)
 	{
-		OutValue = GetProfile()->GetBuildTarget();
+		OutValue = FString::Join( GetProfile()->GetBuildTargets(), TEXT("+") );
 		return true;
 	}
 	else if (InVariable == PlatformNameVariable)
@@ -71,6 +71,7 @@ bool FGlobalsLaunchExtensionInstance::GetExtensionVariableValue( const FString&
 	else if (InVariable == ConfigurationVariable)
 	{
 		OutValue = LexToString( GetProfile()->GetBuildConfiguration() );
+		return true;
 	}
 
 	return false;
diff --git a/Engine/Plugins/Developer/ProjectLauncher/Source/ProjectLauncher/Private/Extension/LaunchExtension.cpp b/Engine/Plugins/Developer/ProjectLauncher/Source/ProjectLauncher/Private/Extension/LaunchExtension.cpp
--- a/Engine/Plugins/Developer/ProjectLauncher/Source/ProjectLauncher/Private/Extension/LaunchExtension.cpp
+++ b/Engine/Plugins/Developer/ProjectLauncher/Source/ProjectLauncher/Private/Extension/LaunchExtension.cpp
@@ -54,7 +54,7 @@ namespace ProjectLauncher
 	}
 
 
-	void ApplyExtensionVariables( const ILauncherProfileRef& InProfile, FString& InOutCommandLine, TSharedRef<FModel> InModel )
+	void ApplyExtensionVariables( const ILauncherProfileRef& InProfile, FString& InOutCommandLine, TSharedRef<FModel> InModel, EBuildTargetType InBuildTargetType )
 	{
 		// take a snapshot of the extension instances for this profile
 		TArray<TSharedPtr<FLaunchExtensionInstance>> ExtensionInstances;
@@ -89,7 +89,14 @@ namespace ProjectLauncher
 				}
 
 				// allow for advanced command line customization
-				ExtensionInstance->CustomizeLaunchCommandLine( InOutCommandLine );
+				if (InBuildTargetType == EBuildTargetType::Unknown)
+				{
+					ExtensionInstance->CustomizeLaunchCommandLine( InOutCommandLine );
+				}
+				else
+				{
+					ExtensionInstance->CustomizeTargetLaunchCommandLine( InOutCommandLine, InBuildTargetType );
+				}
 			}
 		}
 	}
@@ -471,13 +478,35 @@ namespace ProjectLauncher
 
 	FName FLaunchExtensionInstance::GetLaunchPlatformName() const
 	{
-		const PlatformInfo::FTargetPlatformInfo* Info = GetModel()->GetPlatformInfo(GetProfile());
-		if (Info != nullptr)
+		TArray<const PlatformInfo::FTargetPlatformInfo*> Infos = GetModel()->GetPlatformInfos(GetProfile());
+		if (Infos.Num() == 0)
 		{
-			return Info->IniPlatformName;
+			return NAME_None;
 		}
+		if (Infos.Num() == 1)
+		{
+			return Infos[0]->IniPlatformName;
+		}
+
+		auto FindType = [Infos]( EBuildTargetType Type )
+		{
+			for (const PlatformInfo::FTargetPlatformInfo* Info : Infos)
+			{
+				if (Info->PlatformType == Type)
+				{
+					return Info;
+				}
+			}
+
+			return (const PlatformInfo::FTargetPlatformInfo*)nullptr;
+		};
+
+		const PlatformInfo::FTargetPlatformInfo* Result = FindType(EBuildTargetType::Game);
+		Result = Result ? Result : FindType(EBuildTargetType::Client);
+		Result = Result ? Result : FindType(EBuildTargetType::Server);
+		Result = Result ? Result : FindType(EBuildTargetType::Editor);
 
-		return NAME_None;
+		return Result ? Result->IniPlatformName : NAME_None;
 	}
 
 	FString FLaunchExtensionInstance::GetConfigString( EConfig Config, const TCHAR* Name, const TCHAR* DefaultValue ) const
diff --git a/Engine/Plugins/Developer/ProjectLauncher/Source/ProjectLauncher/Private/Model/ProjectLauncherModel.cpp b/Engine/Plugins/Developer/ProjectLauncher/Source/ProjectLauncher/Private/Model/ProjectLauncherModel.cpp
--- a/Engine/Plugins/Developer/ProjectLauncher/Source/ProjectLauncher/Private/Model/ProjectLauncherModel.cpp
+++ b/Engine/Plugins/Developer/ProjectLauncher/Source/ProjectLauncher/Private/Model/ProjectLauncherModel.cpp
@@ -27,10 +27,7 @@
 
 namespace ProjectLauncher
 {
-	// this option filters the platforms & targets based on each other. 
-	// it makes it easier to choose the desired build target but the logic under the hood is more complex
-	bool bUseFriendlyBuildTargetSelection = true;
-
+	const FText ProjectLauncherText = LOCTEXT("ProjectLauncherTitle", "Project Launcher");
 
 	FModel::FModel(const TSharedRef<ITargetDeviceProxyManager>& InDeviceProxyManager, const TSharedRef<ILauncher>& InLauncher, const TSharedRef<ILauncherProfileManager>& InProfileManager)
 		: DeviceProxyManager(InDeviceProxyManager)
@@ -131,7 +128,7 @@ namespace ProjectLauncher
 
 	bool FModel::IsAdvancedProfile(const ILauncherProfileRef& Profile) const
 	{
-		return Profile->GetCookedPlatforms().Num() > 1 ||
+		return 
 				Profile->IsCreatingDLC() ||
 				Profile->IsCreatingReleaseVersion() ||
 				Profile->IsGeneratingPatch() ||				
@@ -202,6 +199,25 @@ namespace ProjectLauncher
 		return nullptr;
 	}
 
+	TArray<const PlatformInfo::FTargetPlatformInfo*> FModel::GetPlatformInfos(const ILauncherProfilePtr& Profile)
+	{
+		TArray<const PlatformInfo::FTargetPlatformInfo*> Result;
+		if (Profile.IsValid())
+		{
+			for (const FString& Platform : Profile->GetCookedPlatforms())
+			{
+				const PlatformInfo::FTargetPlatformInfo* PlatformInfo = PlatformInfo::FindPlatformInfo(FName(Platform));
+				if (PlatformInfo != nullptr)
+				{
+					Result.Add(PlatformInfo);
+				}
+			}
+		}
+
+		return MoveTemp(Result);
+	}
+
+
 	const PlatformInfo::FTargetPlatformInfo* FModel::GetPlatformInfo(FName PlatformName, const FTargetInfo& BuildTargetInfo)
 	{
 		const PlatformInfo::FTargetPlatformInfo* PlatformInfo = PlatformInfo::FindPlatformInfo(PlatformName);
@@ -217,21 +233,27 @@ namespace ProjectLauncher
 		}
 
 		// try to find a matching flavor for the given platform & build target
-		if (PlatformInfo->VanillaInfo->PlatformType == BuildTargetInfo.Type && PlatformInfo->VanillaInfo->PlatformFlavor == PlatformInfo->PlatformFlavor)
-		{
-			return PlatformInfo->VanillaInfo;
-		}
-		for (const PlatformInfo::FTargetPlatformInfo* PlatformInfoFlavor : PlatformInfo->VanillaInfo->Flavors)
+		if (PlatformInfo->VanillaInfo != nullptr)
 		{
-			if (PlatformInfoFlavor->PlatformType == BuildTargetInfo.Type && PlatformInfoFlavor->PlatformFlavor == PlatformInfo->PlatformFlavor)
+			if (PlatformInfo->VanillaInfo->PlatformType == BuildTargetInfo.Type && PlatformInfo->VanillaInfo->PlatformFlavor == PlatformInfo->PlatformFlavor)
+			{
+				return PlatformInfo->VanillaInfo;
+			}
+			for (const PlatformInfo::FTargetPlatformInfo* PlatformInfoFlavor : PlatformInfo->VanillaInfo->Flavors)
 			{
-				return PlatformInfoFlavor;
+				if (PlatformInfoFlavor->PlatformType == BuildTargetInfo.Type && PlatformInfoFlavor->PlatformFlavor == PlatformInfo->PlatformFlavor)
+				{
+					return PlatformInfoFlavor;
+				}
 			}
 		}
 
 		return nullptr;
 	}
 
+
+
+
 	bool FModel::IsHostPlatform(const ILauncherProfilePtr& Profile)
 	{
 		if (Profile.IsValid() && Profile->GetCookedPlatforms().Num() > 0)
@@ -255,6 +277,22 @@ namespace ProjectLauncher
 		return false;
 	}
 
+	bool FModel::IsUsingRemotePlatform(const ILauncherProfilePtr& Profile)
+	{
+		if (Profile.IsValid())
+		{
+			for (const FString& Platform : Profile->GetCookedPlatforms())
+			{
+				if (!IsHostPlatform(FName(Platform)))
+				{
+					return true;
+				}
+			}
+		}
+
+		return false;
+	}
+
 
 
 	FTargetInfo FModel::GetBuildTargetInfo( FString BuildTargetName, const FString& ProjectPath )
@@ -283,15 +321,36 @@ namespace ProjectLauncher
 
 	FTargetInfo FModel::GetBuildTargetInfo( const ILauncherProfileRef& Profile )
 	{
+PRAGMA_DISABLE_DEPRECATION_WARNINGS
 		const FString& BuildTargetName = Profile->GetBuildTarget();
+PRAGMA_ENABLE_DEPRECATION_WARNINGS
 		const FString& ProjectPath = Profile->GetProjectPath();
 		return GetBuildTargetInfo(BuildTargetName, ProjectPath);
 	}
 
+	TArray<FTargetInfo> FModel::GetBuildTargetInfos(const ILauncherProfileRef& Profile)
+	{
+		const FString& ProjectPath = Profile->GetProjectPath();
+		const TArray<FString> BuildTargets = Profile->GetBuildTargets();
+
+		TArray<FTargetInfo> Result;
+		if (BuildTargets.Num() == 0 || BuildTargets.Contains(FString()))
+		{
+			Result.Add( GetBuildTargetInfo(FString(), ProjectPath) );
+		}
+		else for (const FString& BuildTarget : BuildTargets)
+		{
+			Result.Add( GetBuildTargetInfo(BuildTarget, ProjectPath) );
+		}
+
+		return MoveTemp(Result);
+	}
+
+
 	FString FModel::GetVanillaPlatformName( const FString& PlatformName )
 	{
 		const PlatformInfo::FTargetPlatformInfo* PlatformInfo = PlatformInfo::FindPlatformInfo(*PlatformName);
-		if (PlatformInfo != nullptr)
+		if (PlatformInfo != nullptr && PlatformInfo->VanillaInfo != nullptr)
 		{
 			return PlatformInfo->VanillaInfo->Name.ToString();
 		}
@@ -527,27 +586,34 @@ namespace ProjectLauncher
 		if (DeviceProxy.IsValid())
 		{
 			Profile->ClearCookedPlatforms();
-		
-			FTargetInfo BuildTargetInfo = GetBuildTargetInfo(Profile);
-			FString PlatformName = DeviceProxy->GetTargetPlatformName(NAME_None);
-			const PlatformInfo::FTargetPlatformInfo* PlatformInfo = GetPlatformInfo(*PlatformName, BuildTargetInfo);
 
-			if (PlatformInfo != nullptr)
-			{
-				Profile->AddCookedPlatform(PlatformInfo->Name.ToString());
-			}
-			else
+			TArray<FTargetInfo> BuildTargetInfos = GetBuildTargetInfos(Profile);
+			for (const FTargetInfo& BuildTargetInfo : BuildTargetInfos)
 			{
-				Profile->AddCookedPlatform(DeviceProxy->GetTargetPlatformName(NAME_None));
-			}
+				FString PlatformName = DeviceProxy->GetTargetPlatformName(NAME_None);
+				const PlatformInfo::FTargetPlatformInfo* PlatformInfo = GetPlatformInfo(*PlatformName, BuildTargetInfo);
 
-			// @todo: is default deploy platform necessary now?
+				if (PlatformInfo != nullptr)
+				{
+					Profile->AddCookedPlatform(PlatformInfo->Name.ToString());
+				}
+				else
+				{
+					Profile->AddCookedPlatform(DeviceProxy->GetTargetPlatformName(NAME_None));
+				}
+			}
 		}
 	}
 
 	void FModel::UpdateCookedPlatformsFromBuildTarget(const ILauncherProfileRef& Profile)
 	{
-		FTargetInfo BuildTargetInfo = GetBuildTargetInfo(Profile);
+		TArray<FTargetInfo> BuildTargetInfos = GetBuildTargetInfos(Profile);
+
+		if (BuildTargetInfos.Num() != 1)
+		{
+			return;
+		}
+		const FTargetInfo& BuildTargetInfo = BuildTargetInfos[0];
 
 		TArray<FString> Platforms = Profile->GetCookedPlatforms();
 		Profile->ClearCookedPlatforms();
@@ -825,11 +891,11 @@ namespace ProjectLauncher
 
 
 
-	extern void ApplyExtensionVariables( const ILauncherProfileRef& InProfile, FString& InOutCommandLine, TSharedRef<FModel> InModel );
+	extern void ApplyExtensionVariables( const ILauncherProfileRef& InProfile, FString& InOutCommandLine, TSharedRef<FModel> InModel, EBuildTargetType InBuildTargetType );
 
-	void FModel::OnModifyLaunchCommandLine( const ILauncherProfileRef& InProfile, FString& InOutCommandLine )
+	void FModel::OnModifyLaunchCommandLine( const ILauncherProfileRef& InProfile, EBuildTargetType InBuildTargetType, FString& InOutCommandLine )
 	{
-		ApplyExtensionVariables(InProfile, InOutCommandLine, AsShared() );
+		ApplyExtensionVariables(InProfile, InOutCommandLine, AsShared(), InBuildTargetType );
 	}
 
 
diff --git a/Engine/Plugins/Developer/ProjectLauncher/Source/ProjectLauncher/Private/ProfileTree/GenericProfileTreeBuilder.cpp b/Engine/Plugins/Developer/ProjectLauncher/Source/ProjectLauncher/Private/ProfileTree/GenericProfileTreeBuilder.cpp
--- a/Engine/Plugins/Developer/ProjectLauncher/Source/ProjectLauncher/Private/ProfileTree/GenericProfileTreeBuilder.cpp
+++ b/Engine/Plugins/Developer/ProjectLauncher/Source/ProjectLauncher/Private/ProfileTree/GenericProfileTreeBuilder.cpp
@@ -11,6 +11,7 @@
 #include "Widgets/Text/STextBlock.h"
 #include "Widgets/Layout/SBox.h"
 #include "Widgets/SBoxPanel.h"
+#include "Dialog/SMessageDialog.h"
 #include "SResizeBox.h"
 #include "SSearchableComboBox.h"
 #include "PlatformInfo.h"
@@ -19,7 +20,7 @@
 #include "Interfaces/ITargetPlatform.h"
 #include "GameProjectHelper.h"
 #include "InstalledPlatformInfo.h"
-
+#include "HAL/PlatformMisc.h"
 
 
 #define LOCTEXT_NAMESPACE "CustomProfileTreeBuilder"
@@ -98,8 +99,8 @@ namespace ProjectLauncher
 		const TArray<FString>& DeviceIDs = Profile->GetDeployedDeviceGroup()->GetDeviceIDs();
 
 		bShouldCook = GetCook();
+		bHasAdvancedPlatformTargets = (Profile->GetCookedPlatforms().Num() > 1 || Profile->GetBuildTargets().Num() > 1);
 		ContentScheme = Model->DetermineProfileContentScheme(Profile);
-		CachedBuildTargetType = Model->GetBuildTargetInfo(Profile).Type;
 
 		// ensure the default deploy platform matches the cook platform because Legacy Project Launcher can create a mismatch
 		if (Profile->GetDefaultDeployPlatform() != NAME_None && !GetSelectedPlatforms().Contains(Profile->GetDefaultDeployPlatform().ToString()) )
@@ -109,6 +110,7 @@ namespace ProjectLauncher
 		}
 	
 		CacheArchitectures();
+		CacheBuildTargets();
 	}
 
 
@@ -164,13 +166,13 @@ namespace ProjectLauncher
 			// the code is implemented with this syntax to avoid aid readability without using slate's TAttribute style functionality which seemed like an overkill for our simpler needs.
 			HeadingNode.AddWidget( LOCTEXT("TargetLabel", "Target"), 
 				{
-					.IsDefault = [this]()		{ return !Profile->HasBuildTargetSpecified() || Profile->GetBuildTarget().IsEmpty(); },
-					.SetToDefault = [this]()	{ SetBuildTarget(FString()); },
+					.IsDefault = [this]()		{ return !Profile->HasBuildTargetSpecified() || (Profile->GetBuildTargets().Num() == 0); },
+					.SetToDefault = [this]()	{ SetBuildTargets(TArray<FString>()); },
 					.Validation = FValidation( {ELauncherProfileValidationErrors::BuildTargetCookVariantMismatch, ELauncherProfileValidationErrors::BuildTargetIsRequired, ELauncherProfileValidationErrors::CookOnTheFlyDoesntSupportServer} ),
 				},
-				SNew(SCustomLaunchBuildTargetCombo, Model)
-				.OnSelectionChanged(this, &FGenericProfileTreeBuilder::SetBuildTarget)
-				.SelectedBuildTarget(this, &FGenericProfileTreeBuilder::GetBuildTarget)
+				SAssignNew(BuildTargetCombo, SCustomLaunchBuildTargetCombo, Model)
+				.OnSelectionChanged(this, &FGenericProfileTreeBuilder::SetBuildTargets)
+				.SelectedBuildTargets(this, &FGenericProfileTreeBuilder::GetBuildTargets)
 				.SelectedProject(this, &FGenericProfileTreeBuilder::GetProjectPath)
 				.Font(FCoreStyle::Get().GetFontStyle("SmallFont"))
 			);
@@ -179,16 +181,17 @@ namespace ProjectLauncher
 		{
 			HeadingNode.AddWidget( LOCTEXT("TargetLabel", "Target"), 
 				{
-					.IsDefault = [this]()		{ return !Profile->HasBuildTargetSpecified() || Profile->GetBuildTarget().IsEmpty(); },
-					.SetToDefault = [this]()	{ SetBuildTarget(FString()); },
+					.IsDefault = [this]()		{ return !Profile->HasBuildTargetSpecified() || (Profile->GetBuildTargets().Num() == 0); },
+					.SetToDefault = [this]()	{ SetBuildTargets(TArray<FString>()); },
 					.IsEnabled = [this]()		{ return Profile->HasProjectSpecified(); },
-					.Validation = FValidation( {ELauncherProfileValidationErrors::BuildTargetCookVariantMismatch, ELauncherProfileValidationErrors::BuildTargetIsRequired, ELauncherProfileValidationErrors::CookOnTheFlyDoesntSupportServer} ),
+					.Validation = FValidation( {ELauncherProfileValidationErrors::BuildTargetCookVariantMismatch, ELauncherProfileValidationErrors::BuildTargetIsRequired, ELauncherProfileValidationErrors::CookOnTheFlyDoesntSupportServer}, {TEXT("Validation_GameAndClientTargetsDisallowed")}),
 				},
-				SNew(SCustomLaunchBuildTargetCombo, Model)
-				.OnSelectionChanged(this, &FGenericProfileTreeBuilder::SetBuildTarget)
-				.SelectedBuildTarget(this, &FGenericProfileTreeBuilder::GetBuildTarget)
-				.SupportedTargetTypes(this, &FGenericProfileTreeBuilder::GetSupportedBuildTargetTypes)
+				SAssignNew(BuildTargetCombo, SCustomLaunchBuildTargetCombo, Model)
+				.OnSelectionChanged(this, &FGenericProfileTreeBuilder::SetBuildTargets)
+				.SelectedBuildTargets(this, &FGenericProfileTreeBuilder::GetBuildTargets)
 				.SelectedProject(this, &FGenericProfileTreeBuilder::GetProjectPath)
+				//.SetAdvancedPlatformsOption(this, &FGenericProfileTreeBuilder::SetHasAdvancedPlatformTargets) // @todo: prevent this from being toggled until the final design on client/server build configuration. only exising multi-target/multi-platform profiles will have multiselectable platforms/target controls
+				.GetAdvancedPlatformsOption(this, &FGenericProfileTreeBuilder::HasAdvancedPlatformTargets)
 				.Font(FCoreStyle::Get().GetFontStyle("SmallFont"))
 			);
 		}
@@ -200,14 +203,16 @@ namespace ProjectLauncher
 			{
 				.Validation = FValidation({ELauncherProfileValidationErrors::NoPlatformSelected, ELauncherProfileValidationErrors::NoPlatformSDKInstalled})
 			},
-			SNew(SCustomLaunchPlatformCombo)
+			SAssignNew(PlatformCombo, SCustomLaunchPlatformCombo)
 			.SelectedPlatforms(this, &FGenericProfileTreeBuilder::GetSelectedPlatforms)
 			.OnSelectionChanged(this, &FGenericProfileTreeBuilder::SetSelectedPlatforms)
-			.BasicPlatformsOnly(bUseFriendlyBuildTargetSelection)
+			.GetAdvancedPlatformsOption(this, &FGenericProfileTreeBuilder::HasAdvancedPlatformTargets)
+			.SupportedBuildTargets_Lambda( [this]() { return CachedBuildTargetTypes; } )
 			.Font(FCoreStyle::Get().GetFontStyle("SmallFont"))
 		);
 	}
 
+
 	void FGenericProfileTreeBuilder::AddConfigurationProperty( FLaunchProfileTreeNode& HeadingNode )
 	{
 		TArray<EBuildConfiguration> ValidConfigurations;
@@ -400,18 +405,37 @@ namespace ProjectLauncher
 
 	void FGenericProfileTreeBuilder::AddArchitectureProperty( FLaunchProfileTreeNode& HeadingNode )
 	{
-		HeadingNode.AddWidget( LOCTEXT("ArchitectureLabel", "Architecture"),
-			{
-				.IsDefault = [this]()		{ return GetArchitecture().IsEmpty(); },
-				.SetToDefault = [this]()	{ SetArchitecture(FString()); },
-				.IsVisible = [this]()       { return CachedArchitectures.Num() > 0; }
-			},
-			SNew(SCustomLaunchStringCombo)
-			.OnSelectionChanged( this, &FGenericProfileTreeBuilder::SetArchitecture )
-			.SelectedItem( this, &FGenericProfileTreeBuilder::GetArchitecture)
-			.GetDisplayName( this, &FGenericProfileTreeBuilder::GetArchitectureDisplayName )
-			.Items_Lambda( [this]() { return CachedArchitectures;} )
-		);
+		static const struct
+		{
+			EBuildTargetType TargetType;
+			FText Label;
+		}
+		PerTargetTypeArchitectures[] =
+		{
+			{ EBuildTargetType::Game,   LOCTEXT("GameArchitectureLabel",   "Game Architecture")   },
+			{ EBuildTargetType::Client, LOCTEXT("ClientArchitectureLabel", "Client Architecture") },
+			{ EBuildTargetType::Server, LOCTEXT("ServerArchitectureLabel", "Server Architecture") },
+			{ EBuildTargetType::Editor, LOCTEXT("EditorArchitectureLabel", "Edtor Architecture") },
+		};
+
+		for (int32 Index = 0; Index < UE_ARRAY_COUNT(PerTargetTypeArchitectures); Index++)
+		{
+			EBuildTargetType TargetType = PerTargetTypeArchitectures[Index].TargetType;
+
+			HeadingNode.AddWidget( PerTargetTypeArchitectures[Index].Label,
+				{
+					.IsDefault = [this, TargetType]()		{ return GetArchitecture(TargetType).IsEmpty(); },
+					.SetToDefault = [this, TargetType]()	{ SetArchitecture(FString(), TargetType); },
+					.IsVisible = [this, TargetType]()       { return CachedBuildTargetTypes.Contains(TargetType) && CachedBuildTargetArchitectures[TargetType].Num() > 0; }
+				},
+				SNew(SCustomLaunchStringCombo)
+				.OnSelectionChanged( this, &FGenericProfileTreeBuilder::SetArchitecture, TargetType )
+				.SelectedItem( this, &FGenericProfileTreeBuilder::GetArchitecture, TargetType)
+				.GetDisplayName( this, &FGenericProfileTreeBuilder::GetArchitectureDisplayName )
+				.Items_Lambda( [this, TargetType]() { return CachedBuildTargetArchitectures[TargetType];} )
+			);
+
+		}
 	}
 
 
@@ -539,13 +563,45 @@ namespace ProjectLauncher
 	{
 		HeadingNode.AddCommandLineString( LOCTEXT("CommandLineLabel", "Additional Command Line"), 
 			{
-				.GetValue = [this]()				{ return GetCommandLine(); },
-				.SetValue = [this](FString Value)	{ SetCommandLine(Value); },
+				.GetValue = [this]()				{ return GetCommandLine(EBuildTargetType::Unknown); },
+				.SetValue = [this](FString Value)	{ SetCommandLine(Value, EBuildTargetType::Unknown); },
 				.GetDefaultValue = EmptyString,
 				.IsEnabled = ForRun,
 				.Validation = FValidation({ELauncherProfileValidationErrors::MalformedLaunchCommandLine}),
 			}
 		);
+
+
+
+
+		static const struct
+		{
+			EBuildTargetType TargetType;
+			FText Label;
+		}
+		PerTargetTypeCmdLine[] =
+		{
+			{ EBuildTargetType::Game,    LOCTEXT("CommandLineLabelGame",   "Game Command Line")   },
+			{ EBuildTargetType::Client,  LOCTEXT("CommandLineLabelClient", "Client Command Line") },
+			{ EBuildTargetType::Server,  LOCTEXT("CommandLineLabelServer", "Server Command Line") },
+			{ EBuildTargetType::Editor,  LOCTEXT("CommandLineLabelEditor", "Editor Command Line") },
+		};
+
+		for (int32 Index = 0; Index < UE_ARRAY_COUNT(PerTargetTypeCmdLine); Index++)
+		{
+			EBuildTargetType TargetType = PerTargetTypeCmdLine[Index].TargetType;
+
+			HeadingNode.AddString(  PerTargetTypeCmdLine[Index].Label, // cannot use .AddCommandLineString at the moment because extensions can cache state so all + menus only affect common command line for now
+				{
+					.GetValue = [this, TargetType]()				{ return GetCommandLine(TargetType); },
+					.SetValue = [this, TargetType](FString Value)	{ SetCommandLine(Value, TargetType); },
+					.GetDefaultValue = EmptyString,
+					.IsVisible = [this, TargetType]()				{ return HasAdvancedPlatformTargets() && CachedBuildTargetTypes.Contains(TargetType); },
+					.IsEnabled = ForRun,
+				}
+			);
+
+		}
 	}
 
 
@@ -579,11 +635,16 @@ namespace ProjectLauncher
 				Profile->AddCustomError(TEXT("Validation_ZenNeedsIoStore"), LOCTEXT("Error_ZenNeedsIoStore", "Zen Workflows require IoStore to be enabled in your project's Packaging Settings. It cannot be overridden in a profile"));
 			}
 			
-			if (!Model->GetProjectSettings(Profile).bAllowRemoteNetworkService && !Model->IsHostPlatform(Profile))
+			if (!Model->GetProjectSettings(Profile).bAllowRemoteNetworkService && Model->IsUsingRemotePlatform(Profile))
 			{
 				Profile->AddCustomError(TEXT("Validation_ZenNeedsRemotePermission"), LOCTEXT("Error_ZenNeedsRemotePermission", "Zen Workflows for remote devices need [Zen.AutoLaunch]:AllowRemoteNetworkService=true. It cannot be overridden in a profile"));
 			}
 		}
+
+		if (CachedBuildTargetTypes.Contains(EBuildTargetType::Game) && CachedBuildTargetTypes.Contains(EBuildTargetType::Client))
+		{
+			Profile->AddCustomError(TEXT("Validation_GameAndClientTargetsDisallowed"), LOCTEXT("Error_NoGameAndClient", "Cannot specify Game and Client targets to be built together"));
+		}
 	}
 
 
@@ -626,21 +687,56 @@ namespace ProjectLauncher
 
 	void FGenericProfileTreeBuilder::CacheArchitectures()
 	{
-		CachedArchitectures.Reset();
+		CachedBuildTargetArchitectures.Reset();
+		CachedBuildTargetArchitectures.Add(EBuildTargetType::Game);
+		CachedBuildTargetArchitectures.Add(EBuildTargetType::Client);
+		CachedBuildTargetArchitectures.Add(EBuildTargetType::Server);
+		CachedBuildTargetArchitectures.Add(EBuildTargetType::Editor);
 
-		const PlatformInfo::FTargetPlatformInfo* PlatformInfo = Model->GetPlatformInfo(Profile);
-		if (PlatformInfo != nullptr)
+		// gather all architecures for each target type
+		for (const PlatformInfo::FTargetPlatformInfo* PlatformInfo : Model->GetPlatformInfos(Profile))
 		{
 			const ITargetPlatform* TargetPlatform = GetTargetPlatformManager()->FindTargetPlatform(PlatformInfo->IniPlatformName);
 			if (TargetPlatform)
 			{
-				TargetPlatform->GetPossibleArchitectures(CachedArchitectures);
-				if (CachedArchitectures.Num() > 0)
+				TArray<FString> PossibleArchitectures;
+				TargetPlatform->GetPossibleArchitectures(PossibleArchitectures);
+
+				for (const FString& PossibleArchitecture : PossibleArchitectures)
 				{
-					CachedArchitectures.Insert(FString(), 0); // empty string for "project default" option
+					CachedBuildTargetArchitectures[PlatformInfo->PlatformType].AddUnique(PossibleArchitecture);
 				}
 			}
 		}
+
+		// insert an empty string at the top of the list for "project default" if there are any architectures
+		auto InsertProjectDefault = []( TArray<FString>& Architectures )
+		{
+			if (Architectures.Num() > 0)
+			{
+				Architectures.Insert(FString(), 0);
+			}
+		};
+		InsertProjectDefault( CachedBuildTargetArchitectures[EBuildTargetType::Game] );
+		InsertProjectDefault( CachedBuildTargetArchitectures[EBuildTargetType::Client] );
+		InsertProjectDefault( CachedBuildTargetArchitectures[EBuildTargetType::Server] );
+		InsertProjectDefault( CachedBuildTargetArchitectures[EBuildTargetType::Editor] );
+	}
+
+	void FGenericProfileTreeBuilder::CacheBuildTargets()
+	{
+		CachedBuildTargetTypes.Reset();
+
+		TArray<FTargetInfo> TargetInfos = Model->GetBuildTargetInfos(Profile);
+		for (const FTargetInfo& TargetInfo : TargetInfos)
+		{
+			CachedBuildTargetTypes.Add(TargetInfo.Type);
+		}
+
+		if (bHasAdvancedPlatformTargets && PlatformCombo.IsValid())
+		{
+			PlatformCombo->RefreshPlatformsList();
+		}
 	}
 
 
@@ -657,22 +753,24 @@ namespace ProjectLauncher
 
 	void FGenericProfileTreeBuilder::SetSelectedPlatforms( TArray<FString> SelectedPlatforms )
 	{
-		Profile->ClearCookedPlatforms();
-
-		if (bUseFriendlyBuildTargetSelection)
+		if (bHasAdvancedPlatformTargets)
 		{
-			FTargetInfo BuildTargetInfo = Model->GetBuildTargetInfo(GetBuildTarget(), GetProjectPath());
-
+			Profile->ClearCookedPlatforms();
 			for (const FString& Platform : SelectedPlatforms)
 			{
-				Profile->AddCookedPlatform( FModel::GetBuildTargetPlatformName(Platform, BuildTargetInfo));
+				Profile->AddCookedPlatform(Platform);
 			}
 		}
 		else
 		{
+			Profile->ClearCookedPlatforms();
+
+			FString BuildTarget = (GetBuildTargets().Num() > 0) ? GetBuildTargets()[0] : FString();
+			FTargetInfo BuildTargetInfo = Model->GetBuildTargetInfo(BuildTarget, GetProjectPath());
+
 			for (const FString& Platform : SelectedPlatforms)
 			{
-				Profile->AddCookedPlatform(Platform);
+				Profile->AddCookedPlatform( FModel::GetBuildTargetPlatformName(Platform, BuildTargetInfo));
 			}
 		}
 
@@ -682,9 +780,14 @@ namespace ProjectLauncher
 			Profile->SetDefaultDeployPlatform(FName(*DeployPlatform));
 		}
 
+		ValidateArchitectures();
+
 		OnPropertyChanged();
 
-		DeployDeviceListView->OnSelectedPlatformChanged();
+		if (DeployDeviceListView)
+		{
+			DeployDeviceListView->OnSelectedPlatformChanged();
+		}
 
 		CacheArchitectures();
 	}
@@ -693,17 +796,17 @@ namespace ProjectLauncher
 	{
 		TArray<FString> Platforms;
 
-		if (bUseFriendlyBuildTargetSelection)
+		if (bHasAdvancedPlatformTargets)
+		{
+			Platforms = Profile->GetCookedPlatforms();
+		}
+		else
 		{
 			for (const FString& Platform : Profile->GetCookedPlatforms())
 			{
 				Platforms.Add(FModel::GetPlatformNameWithFlavor(Platform));
 			}
 		}
-		else
-		{
-			Platforms = Profile->GetCookedPlatforms();
-		}
 
 		return MoveTemp(Platforms);
 	}
@@ -719,7 +822,7 @@ namespace ProjectLauncher
 		Profile->SetProjectSpecified(!ProjectPath.IsEmpty());
 		Profile->SetProjectPath(ProjectPath);
 
-		if (bUseFriendlyBuildTargetSelection && ProfileType != EProfileType::Basic )
+		if (ProfileType != EProfileType::Basic )
 		{
 			Model->UpdateCookedPlatformsFromBuildTarget(Profile);
 		}
@@ -735,71 +838,63 @@ namespace ProjectLauncher
 	}
 
 
-	FString FGenericProfileTreeBuilder::GetBuildTarget() const
+	TArray<FString> FGenericProfileTreeBuilder::GetBuildTargets() const
 	{
 		if (ProfileType == EProfileType::Basic)
 		{
+			check(!bHasAdvancedPlatformTargets);
 			if (Profile->HasBuildTargetSpecified())
 			{
-				return Profile->GetBuildTarget();
+				return Profile->GetBuildTargets();
 			}
 			else
 			{
-				return Model->GetProfileManager()->GetBuildTarget();
+				return TArray<FString>{Model->GetProfileManager()->GetBuildTarget()};
 			}
 		}
 		else
 		{
-			return Profile->GetBuildTarget();
+			return Profile->GetBuildTargets();
 		}
 	}
-
-	void FGenericProfileTreeBuilder::SetBuildTarget(FString BuildTarget)
+		
+	void FGenericProfileTreeBuilder::SetBuildTargets( TArray<FString> BuildTargets)
 	{
-		Profile->SetBuildTargetSpecified(!BuildTarget.IsEmpty());
-		Profile->SetBuildTarget(BuildTarget);
-	
-		if (ProfileType == EProfileType::Basic )
+		Profile->ClearBuildTargets();
+		if (BuildTargets.Num() == 0 || BuildTargets.Contains(FString()))
 		{
-			Model->UpdatedCookedPlatformsFromDeployDeviceProxy(Profile);
+			Profile->SetBuildTargetSpecified(false);
 		}
-		else if (bUseFriendlyBuildTargetSelection)
+		else
 		{
-			Model->UpdateCookedPlatformsFromBuildTarget(Profile);
-		}
+			Profile->SetBuildTargetSpecified(true);
 
-		CachedBuildTargetType = Model->GetBuildTargetInfo(Profile).Type;
-		OnPropertyChanged();
-	}
+			for (const FString& BuildTarget : BuildTargets)
+			{
+				Profile->AddBuildTarget(BuildTarget);
+			}
+		}
+		CacheBuildTargets();
 
-	TArray<EBuildTargetType> FGenericProfileTreeBuilder::GetSupportedBuildTargetTypes() const
-	{
-		TArray<EBuildTargetType> Result;
 
-		if (Profile->GetCookedPlatforms().Num() > 0)
+		if (!bHasAdvancedPlatformTargets)
 		{
-			const PlatformInfo::FTargetPlatformInfo* PlatformInfo = PlatformInfo::FindPlatformInfo(FName(Profile->GetCookedPlatforms()[0]));
-			if (PlatformInfo)
+			if (ProfileType == EProfileType::Basic )
 			{
-				if (bUseFriendlyBuildTargetSelection)
-				{
-					Result.Add(PlatformInfo->VanillaInfo->PlatformType);
-
-					for (const PlatformInfo::FTargetPlatformInfo* PlatformFlavorInfo : PlatformInfo->VanillaInfo->Flavors)
-					{
-						Result.AddUnique(PlatformFlavorInfo->PlatformType);
-					}
-				}
-				else
-				{
-					Result.Add(PlatformInfo->PlatformType);
-				}
+				Model->UpdatedCookedPlatformsFromDeployDeviceProxy(Profile);
+			}
+			else
+			{
+				Model->UpdateCookedPlatformsFromBuildTarget(Profile);
 			}
 		}
 
-		return MoveTemp(Result);
+		ValidateArchitectures();
+
+		OnPropertyChanged();
 	}
 
+
 	void FGenericProfileTreeBuilder::SetBuildConfiguration(EBuildConfiguration BuildConfiguration)
 	{
 		Profile->SetBuildConfiguration(BuildConfiguration);
@@ -858,7 +953,7 @@ namespace ProjectLauncher
 	
 		// cannot launch via Zen if we're targeting a remote device and Zen isn't accepting external connections
 		if ((InContentScheme == EContentScheme::ZenStreaming || InContentScheme == EContentScheme::ZenPakStreaming) 
-			&& !ProjectSettings.bAllowRemoteNetworkService && !Model->IsHostPlatform(Profile))
+			&& !ProjectSettings.bAllowRemoteNetworkService && Model->IsUsingRemotePlatform(Profile))
 		{
 			OutReason = LOCTEXT("NoZenReason", "Zen Streaming to a remote device requires AllowRemoteNetworkService");
 			return false;
@@ -869,20 +964,34 @@ namespace ProjectLauncher
 
 
 
-	FString FGenericProfileTreeBuilder::GetCommandLine() const
+	FString FGenericProfileTreeBuilder::GetCommandLine(EBuildTargetType BuildTarget) const
 	{
-		// get unified command line from these two fields. the first is presented in "Build" for old Project Launcher, and the latter is presented in "Launch" for old Project Launcher)
-		// when we save back to the profile, this will be stored just in the "Build" one for clarity (because multiple roles are not supported in old or new project launcher)
-		FString CommandLine = Profile->GetAdditionalCommandLineParameters().TrimStartAndEnd() + TEXT(" ") + Profile->GetDefaultLaunchRole()->GetUATCommandLine().TrimStartAndEnd();
-		CommandLine.TrimStartAndEndInline();
+		if (BuildTarget == EBuildTargetType::Unknown)
+		{
+			// get unified command line from these two fields. the first is presented in "Build" for old Project Launcher, and the latter is presented in "Launch" for old Project Launcher)
+			// when we save back to the profile, this will be stored just in the "Build" one for clarity (because multiple roles are not supported in old or new project launcher)
+			FString CommandLine = Profile->GetAdditionalCommandLineParameters().TrimStartAndEnd() + TEXT(" ") + Profile->GetDefaultLaunchRole()->GetUATCommandLine().TrimStartAndEnd();
+			CommandLine.TrimStartAndEndInline();
 
-		return CommandLine;
+			return CommandLine;
+		}
+		else
+		{
+			return Profile->GetAdditionalTargetCommandLineParameters(BuildTarget).TrimStartAndEnd();
+		}
 	}
 
-	void FGenericProfileTreeBuilder::SetCommandLine( const FString& NewCommandLine )
+	void FGenericProfileTreeBuilder::SetCommandLine( const FString& NewCommandLine, EBuildTargetType BuildTarget )
 	{
-		Profile->SetAdditionalCommandLineParameters(NewCommandLine);
-		Profile->GetDefaultLaunchRole()->SetCommandLine(TEXT(""));
+		if (BuildTarget == EBuildTargetType::Unknown)
+		{
+			Profile->SetAdditionalCommandLineParameters(NewCommandLine);
+			Profile->GetDefaultLaunchRole()->SetCommandLine(TEXT(""));
+		}
+		else
+		{
+			Profile->SetAdditionalTargetCommandLineParameters(NewCommandLine, BuildTarget);
+		}
 		OnPropertyChanged();
 	}
 
@@ -1001,7 +1110,7 @@ namespace ProjectLauncher
 			Profile->GetDeployedDeviceGroup()->AddDevice(DeviceID);
 		}
 
-		if (ProfileType == EProfileType::Basic) // @fixme: should this be done for custom profiles too?
+		if (ProfileType == EProfileType::Basic)
 		{
 			Model->UpdatedCookedPlatformsFromDeployDeviceProxy(Profile);
 		}
@@ -1065,7 +1174,10 @@ namespace ProjectLauncher
 
 		OnPropertyChanged();
 
-		DeployDeviceListView->RefreshDeviceList();
+		if (DeployDeviceListView)
+		{
+			DeployDeviceListView->RefreshDeviceList();
+		}
 	}
 
 
@@ -1123,42 +1235,44 @@ namespace ProjectLauncher
 		return (InProfile->GetBuildMode() == ELauncherProfileBuildModes::Build);
 	}
 
-	void FGenericProfileTreeBuilder::SetArchitecture( FString Architecture )
+	void FGenericProfileTreeBuilder::SetArchitecture( FString Architecture, EBuildTargetType BuildTargetType )
 	{
-		// clear existing architectures
-		TArray<FString> Empty;
-		Profile->SetServerArchitectures(Empty);
-		Profile->SetEditorArchitectures(Empty);
-		Profile->SetClientArchitectures(Empty);
-
-		// set new single architecture
+		// Project Launcher only supports launching for one architecture at a time at the moment
+		TArray<FString> Architectures;
 		if (!Architecture.IsEmpty())
 		{
-			TArray<FString> Single;
-			Single.Add(Architecture);
+			Architectures.Add(Architecture);
+		}
 
-			switch (CachedBuildTargetType)
-			{
-				case EBuildTargetType::Server: Profile->SetServerArchitectures(Single); break;
-				case EBuildTargetType::Editor: Profile->SetEditorArchitectures(Single); break;
-				default:                       Profile->SetClientArchitectures(Single); break;
-			}
+		// update the architecture
+		if (BuildTargetType == EBuildTargetType::Game || BuildTargetType == EBuildTargetType::Client)
+		{
+			Profile->SetClientArchitectures(Architectures);
+		}
+		else if (BuildTargetType == EBuildTargetType::Server)
+		{
+			Profile->SetServerArchitectures(Architectures);
+		}
+		else if (BuildTargetType == EBuildTargetType::Editor)
+		{
+			Profile->SetEditorArchitectures(Architectures);
 		}
 
 		OnPropertyChanged();
 	}
 
-	FString FGenericProfileTreeBuilder::GetArchitecture() const
+	FString FGenericProfileTreeBuilder::GetArchitecture(EBuildTargetType BuildTargetType) const
 	{
+		// read current architectures
 		TArray<FString> Architectures;
-
-		switch (CachedBuildTargetType)
+		switch (BuildTargetType)
 		{
 			case EBuildTargetType::Server: Architectures = Profile->GetServerArchitectures(); break;
 			case EBuildTargetType::Editor: Architectures = Profile->GetEditorArchitectures(); break;
 			default:                       Architectures = Profile->GetClientArchitectures(); break;
 		}
 
+		// Project Launcher only supports launching for one architecture at a time at the moment
 		return (Architectures.Num() > 0) ? Architectures[0] : FString();
 	}
 
@@ -1168,7 +1282,7 @@ namespace ProjectLauncher
 		{
 			return LOCTEXT("DefaultArchName", "Project Default");
 		}
-		else if (Architecture == FPlatformMisc::GetHostArchitecture() && Model->IsHostPlatform(Profile))
+		else if (Architecture == FPlatformMisc::GetHostArchitecture())
 		{
 			return FText::Format(LOCTEXT("HostArchLabel", "{0} (this platform)"), FText::FromString(Architecture) );
 		}
@@ -1178,6 +1292,25 @@ namespace ProjectLauncher
 		}
 	}
 
+	void FGenericProfileTreeBuilder::ValidateArchitectures()
+	{
+		CacheArchitectures();
+
+		// if the current build targets don't support architectures, remove the
+
+		if (CachedBuildTargetArchitectures[EBuildTargetType::Game].Num() == 0 && CachedBuildTargetArchitectures[EBuildTargetType::Client].Num() == 0)
+		{
+			Profile->SetClientArchitectures(TArray<FString>());
+		}
+		if (CachedBuildTargetArchitectures[EBuildTargetType::Server].Num() == 0)
+		{
+			Profile->SetServerArchitectures(TArray<FString>());
+		}
+		if (CachedBuildTargetArchitectures[EBuildTargetType::Editor].Num() == 0)
+		{
+			Profile->SetEditorArchitectures(TArray<FString>());
+		}
+	}
 
 
 	void FGenericProfileTreeBuilder::SetDeployToDevice( bool bDeployToDevice )
@@ -1248,6 +1381,62 @@ namespace ProjectLauncher
 		}
 	}
 
+
+	void FGenericProfileTreeBuilder::SetHasAdvancedPlatformTargets( bool bEnable )
+	{
+		if (bHasAdvancedPlatformTargets && !bEnable)
+		{
+			if (Profile->GetCookedPlatforms().Num() > 1 || Profile->GetBuildTargets().Num() > 1)
+			{
+				FText MessageText = LOCTEXT("ClearAdvancedPlatformTargets", "You have multiple platform/targets selected - are you sure you want to go back to single platform/target??");
+				TSharedRef<SMessageDialog> Dialog = SNew(SMessageDialog)
+					.Title(ProjectLauncherText)
+					.Message(MessageText)
+					.Buttons({SMessageDialog::FButton(LOCTEXT("YES","Yes")), SMessageDialog::FButton(LOCTEXT("NO","No"))})
+				;
+				if (Dialog->ShowModal() != 0)
+				{
+					return;
+				}
+			}
+
+			bHasAdvancedPlatformTargets = bEnable;
+
+			if (Profile->GetBuildTargets().Num() > 1)
+			{
+				FString BuildTarget = Profile->GetBuildTargets()[0];
+				Profile->ClearBuildTargets();
+				Profile->AddBuildTarget(BuildTarget);
+
+			}
+			Model->UpdateCookedPlatformsFromBuildTarget(Profile);
+
+			if (Profile->GetCookedPlatforms().Num() > 0)
+			{
+				TArray<FString> Platforms{Profile->GetCookedPlatforms()[0]};
+				SetSelectedPlatforms(Platforms);
+			}
+
+			CacheBuildTargets();
+		}
+		else
+		{
+			bHasAdvancedPlatformTargets = bEnable;
+			OnPropertyChanged();
+		}
+
+		if (PlatformCombo)
+		{
+			PlatformCombo->RefreshPlatformsList();
+		}
+	}
+
+	bool FGenericProfileTreeBuilder::HasAdvancedPlatformTargets() const
+	{
+		return bHasAdvancedPlatformTargets;
+	}
+
+
 }
 
 #undef LOCTEXT_NAMESPACE
diff --git a/Engine/Plugins/Developer/ProjectLauncher/Source/ProjectLauncher/Private/Styling/ProjectLauncherStyle.cpp b/Engine/Plugins/Developer/ProjectLauncher/Source/ProjectLauncher/Private/Styling/ProjectLauncherStyle.cpp
--- a/Engine/Plugins/Developer/ProjectLauncher/Source/ProjectLauncher/Private/Styling/ProjectLauncherStyle.cpp
+++ b/Engine/Plugins/Developer/ProjectLauncher/Source/ProjectLauncher/Private/Styling/ProjectLauncherStyle.cpp
@@ -190,3 +190,16 @@ const FSlateBrush* FProjectLauncherStyle::GetProfileBrushForPlatform(const Platf
 		return GetBrush("Profile.NoPlatform.Large");
 	}
 }
+
+const FSlateBrush* FProjectLauncherStyle::GetProfileBrushForPlatforms(TArray<const PlatformInfo::FTargetPlatformInfo*> PlatformInfos, EPlatformIconSize IconSize)
+{
+	if (PlatformInfos.Num() == 0)
+	{
+		return GetProfileBrushForPlatform(nullptr, IconSize);
+	}
+	else
+	{
+		// @todo: should we have a dedicated icon for multi-platform?
+		return GetProfileBrushForPlatform(PlatformInfos[0], IconSize);
+	}
+}
diff --git a/Engine/Plugins/Developer/ProjectLauncher/Source/ProjectLauncher/Private/Styling/ProjectLauncherStyle.h b/Engine/Plugins/Developer/ProjectLauncher/Source/ProjectLauncher/Private/Styling/ProjectLauncherStyle.h
--- a/Engine/Plugins/Developer/ProjectLauncher/Source/ProjectLauncher/Private/Styling/ProjectLauncherStyle.h
+++ b/Engine/Plugins/Developer/ProjectLauncher/Source/ProjectLauncher/Private/Styling/ProjectLauncherStyle.h
@@ -26,6 +26,7 @@ public:
 
 	static const FSlateBrush* GetBrushForTask(ILauncherTaskPtr Task);
 	static const FSlateBrush* GetProfileBrushForPlatform(const PlatformInfo::FTargetPlatformInfo* PlatformInfo, EPlatformIconSize IconSize);
+	static const FSlateBrush* GetProfileBrushForPlatforms(TArray<const PlatformInfo::FTargetPlatformInfo*> PlatformInfos, EPlatformIconSize IconSize);
 
 private:
 	static TSharedPtr<FSlateStyleSet> StyleSet;	
diff --git a/Engine/Plugins/Developer/ProjectLauncher/Source/ProjectLauncher/Private/Widgets/Output/SCustomLaunchLaunchPanel.cpp b/Engine/Plugins/Developer/ProjectLauncher/Source/ProjectLauncher/Private/Widgets/Output/SCustomLaunchLaunchPanel.cpp
--- a/Engine/Plugins/Developer/ProjectLauncher/Source/ProjectLauncher/Private/Widgets/Output/SCustomLaunchLaunchPanel.cpp
+++ b/Engine/Plugins/Developer/ProjectLauncher/Source/ProjectLauncher/Private/Widgets/Output/SCustomLaunchLaunchPanel.cpp
@@ -391,7 +391,7 @@ void SCustomLaunchLaunchPanel::HandleOutputReceived(const FString& InMessage)
 	{
 		Verbosity = ELogVerbosity::Fatal;
 	}
-	else if ( InMessage.Contains(TEXT("Error:"), ESearchCase::IgnoreCase) )
+	else if ( InMessage.Contains(TEXT("Error:"), ESearchCase::IgnoreCase) || InMessage.StartsWith(TEXT("Result: Failed "), ESearchCase::IgnoreCase ) )
 	{
 		Verbosity = ELogVerbosity::Error;
 	}
@@ -434,8 +434,9 @@ const FSlateBrush* SCustomLaunchLaunchPanel::GetProfileImage() const
 	ILauncherProfilePtr LauncherProfile = GetLauncherProfile();
 	if (LauncherProfile != nullptr)
 	{
-		const PlatformInfo::FTargetPlatformInfo* PlatformInfo = Model->GetPlatformInfo(LauncherProfile);
-		return FProjectLauncherStyle::GetProfileBrushForPlatform(PlatformInfo, EPlatformIconSize::Large);
+		// @todo: should there be a dedicated icon for multi-platform?
+		TArray<const PlatformInfo::FTargetPlatformInfo*> PlatformInfos = Model->GetPlatformInfos(LauncherProfile);
+		return FProjectLauncherStyle::GetProfileBrushForPlatforms(PlatformInfos, EPlatformIconSize::Large);
 	}
 
 	return FStyleDefaults::GetNoBrush();
@@ -502,14 +503,14 @@ FText SCustomLaunchLaunchPanel::GetProfileTargetName() const
 	ILauncherProfilePtr LauncherProfile = GetLauncherProfile();
 	if (LauncherProfile != nullptr)
 	{
-		FString BuildTarget = LauncherProfile->GetBuildTarget();
-		if (BuildTarget.IsEmpty())
+		TArray<FString> BuildTargets = LauncherProfile->GetBuildTargets();
+		if (BuildTargets.Num() == 0 || BuildTargets.Contains(FString()))
 		{
 			return LOCTEXT("DefaultTargetName", "Target Default");
 		}
 		else
 		{
-			return FText::FromString(BuildTarget);
+			return FText::FromString(FString::Join(BuildTargets, TEXT(", ")));
 		}
 	}
 
diff --git a/Engine/Plugins/Developer/ProjectLauncher/Source/ProjectLauncher/Private/Widgets/Profiles/SCustomLaunchCustomProfileSelector.cpp b/Engine/Plugins/Developer/ProjectLauncher/Source/ProjectLauncher/Private/Widgets/Profiles/SCustomLaunchCustomProfileSelector.cpp
--- a/Engine/Plugins/Developer/ProjectLauncher/Source/ProjectLauncher/Private/Widgets/Profiles/SCustomLaunchCustomProfileSelector.cpp
+++ b/Engine/Plugins/Developer/ProjectLauncher/Source/ProjectLauncher/Private/Widgets/Profiles/SCustomLaunchCustomProfileSelector.cpp
@@ -331,8 +331,8 @@ const FSlateBrush* SCustomLaunchCustomProfileSelector::GetProfileImage(ILauncher
 	}
 	else
 	{
-		const PlatformInfo::FTargetPlatformInfo* PlatformInfo = Model->GetPlatformInfo(Profile);
-		return FProjectLauncherStyle::GetProfileBrushForPlatform(PlatformInfo, EPlatformIconSize::Large);
+		TArray<const PlatformInfo::FTargetPlatformInfo*> PlatformInfos = Model->GetPlatformInfos(Profile);
+		return FProjectLauncherStyle::GetProfileBrushForPlatforms(PlatformInfos, EPlatformIconSize::Large);
 	}
 }
 
diff --git a/Engine/Plugins/Developer/ProjectLauncher/Source/ProjectLauncher/Private/Widgets/Profiles/SCustomLaunchProfilesPanel.cpp b/Engine/Plugins/Developer/ProjectLauncher/Source/ProjectLauncher/Private/Widgets/Profiles/SCustomLaunchProfilesPanel.cpp
--- a/Engine/Plugins/Developer/ProjectLauncher/Source/ProjectLauncher/Private/Widgets/Profiles/SCustomLaunchProfilesPanel.cpp
+++ b/Engine/Plugins/Developer/ProjectLauncher/Source/ProjectLauncher/Private/Widgets/Profiles/SCustomLaunchProfilesPanel.cpp
@@ -160,16 +160,6 @@ TSharedRef<SWidget> SCustomLaunchProfilesPanel::CreateProfilesPanel()
 				.ShowAnyProjectOption(false)
 			]
 
-			// Default Build Target selector label
-			+SHorizontalBox::Slot()
-			.AutoWidth()
-			.VAlign(VAlign_Center)
-			.Padding(16,0,2,0)
-			[
-				SNew(STextBlock)
-				.Text(LOCTEXT("DefaultBuildTargetSelectLabel", "Target"))
-				.Visibility(ProjectLauncher::bUseFriendlyBuildTargetSelection ? EVisibility::Collapsed : EVisibility::Visible)
-			]
 
 			// Default Build Target selector
 			+SHorizontalBox::Slot()
@@ -179,9 +169,8 @@ TSharedRef<SWidget> SCustomLaunchProfilesPanel::CreateProfilesPanel()
 			[
 				SNew(SCustomLaunchBuildTargetCombo, Model.ToSharedRef())
 				.OnSelectionChanged(this, &SCustomLaunchProfilesPanel::SetDefaultBuildTarget)
-				.SelectedBuildTarget(this, &SCustomLaunchProfilesPanel::GetDefaultBuildTarget)
+				.SelectedBuildTargets(this, &SCustomLaunchProfilesPanel::GetDefaultBuildTarget)
 				.SelectedProject(this, &SCustomLaunchProfilesPanel::GetDefaultProjectPath)
-				.Visibility(ProjectLauncher::bUseFriendlyBuildTargetSelection ? EVisibility::Collapsed : EVisibility::Visible)
 			]
 
 			// Spacing
@@ -781,8 +770,8 @@ const FSlateBrush* SCustomLaunchProfilesPanel::GetSelectedProfileImage() const
 	}
 	else
 	{
-		const PlatformInfo::FTargetPlatformInfo* PlatformInfo = Model->GetPlatformInfo(Model->GetSelectedProfile());
-		return FProjectLauncherStyle::GetProfileBrushForPlatform(PlatformInfo, EPlatformIconSize::Large);
+		TArray<const PlatformInfo::FTargetPlatformInfo*> PlatformInfos = Model->GetPlatformInfos(Model->GetSelectedProfile());
+		return FProjectLauncherStyle::GetProfileBrushForPlatforms(PlatformInfos, EPlatformIconSize::Large);
 	}
 }
 
@@ -885,8 +874,10 @@ FString SCustomLaunchProfilesPanel::GetDefaultProjectPath() const
 	return Model->GetProfileManager()->GetProjectPath();
 }
 
-void SCustomLaunchProfilesPanel::SetDefaultBuildTarget(FString BuildTarget)
+void SCustomLaunchProfilesPanel::SetDefaultBuildTarget(TArray<FString> BuildTargets)
 {
+	FString BuildTarget = (BuildTargets.Num() > 0) ? BuildTargets[0] : FString(); 
+
 	Model->GetProfileManager()->SetBuildTarget(BuildTarget);
 
 	if (Model->GetSelectedProfile().IsValid())
@@ -895,9 +886,9 @@ void SCustomLaunchProfilesPanel::SetDefaultBuildTarget(FString BuildTarget)
 	}
 }
 
-FString SCustomLaunchProfilesPanel::GetDefaultBuildTarget() const
+TArray<FString> SCustomLaunchProfilesPanel::GetDefaultBuildTarget() const
 {
-	return Model->GetProfileManager()->GetBuildTarget();
+	return TArray<FString>{Model->GetProfileManager()->GetBuildTarget()};
 }
 
 
diff --git a/Engine/Plugins/Developer/ProjectLauncher/Source/ProjectLauncher/Private/Widgets/Profiles/SCustomLaunchProfilesPanel.h b/Engine/Plugins/Developer/ProjectLauncher/Source/ProjectLauncher/Private/Widgets/Profiles/SCustomLaunchProfilesPanel.h
--- a/Engine/Plugins/Developer/ProjectLauncher/Source/ProjectLauncher/Private/Widgets/Profiles/SCustomLaunchProfilesPanel.h
+++ b/Engine/Plugins/Developer/ProjectLauncher/Source/ProjectLauncher/Private/Widgets/Profiles/SCustomLaunchProfilesPanel.h
@@ -73,8 +73,8 @@ private:
 	void SetDefaultProjectPath(FString ProjectPath);
 	FString GetDefaultProjectPath() const;
 
-	void SetDefaultBuildTarget(FString BuildTarget);
-	FString GetDefaultBuildTarget() const;
+	void SetDefaultBuildTarget(TArray<FString> BuildTarget);
+	TArray<FString> GetDefaultBuildTarget() const;
 
 	ProjectLauncher::FOnProfileClicked OnProfileLaunchClicked;
 	TSharedPtr<ProjectLauncher::FModel> Model;
diff --git a/Engine/Plugins/Developer/ProjectLauncher/Source/ProjectLauncher/Private/Widgets/Shared/SCustomLaunchBuildTargetCombo.cpp b/Engine/Plugins/Developer/ProjectLauncher/Source/ProjectLauncher/Private/Widgets/Shared/SCustomLaunchBuildTargetCombo.cpp
--- a/Engine/Plugins/Developer/ProjectLauncher/Source/ProjectLauncher/Private/Widgets/Shared/SCustomLaunchBuildTargetCombo.cpp
+++ b/Engine/Plugins/Developer/ProjectLauncher/Source/ProjectLauncher/Private/Widgets/Shared/SCustomLaunchBuildTargetCombo.cpp
@@ -17,15 +17,16 @@ void SCustomLaunchBuildTargetCombo::Construct(const FArguments& InArgs, TSharedR
 {
 	OnSelectionChanged = InArgs._OnSelectionChanged;
 	SelectedProject = InArgs._SelectedProject;
-	SelectedBuildTarget = InArgs._SelectedBuildTarget;
-	SupportedTargetTypes = InArgs._SupportedTargetTypes;
+	SelectedBuildTargets = InArgs._SelectedBuildTargets;
+	SetAdvancedPlatformsOption = InArgs._SetAdvancedPlatformsOption;
+	GetAdvancedPlatformsOption = InArgs._GetAdvancedPlatformsOption;
 	Model = InModel;
 
 	FSlateFontInfo Font = InArgs._Font.IsSet() ? InArgs._Font.Get() : InArgs._TextStyle->Font;
 
 	ChildSlot
 	[
-		SNew(SComboButton)
+		SAssignNew(ComboButton, SComboButton)
 		.ButtonContent()
 		[
 			SNew(STextBlock)
@@ -42,8 +43,11 @@ END_SLATE_FUNCTION_BUILD_OPTIMIZATION
 BEGIN_SLATE_FUNCTION_BUILD_OPTIMIZATION
 TSharedRef<SWidget> SCustomLaunchBuildTargetCombo::MakeBuildTargetSelectionWidget()
 {
-	FMenuBuilder MenuBuilder(true, nullptr);
+	bool bAutoClose = !GetAdvancedPlatformsOption.IsSet() || !GetAdvancedPlatformsOption.Get(); // only keep the window open if we are in advanced mode (i.e. to allow multiselect)
+	FMenuBuilder MenuBuilder(bAutoClose, nullptr);
 	{
+		EUserInterfaceActionType UIType = GetAdvancedPlatformsOption.Get() ? EUserInterfaceActionType::Check : EUserInterfaceActionType::RadioButton;
+
 		FString DefaultBuildTarget = GetDefaultBuildTargetName();
 		FText DefaultBuildTargetName = DefaultBuildTarget.IsEmpty() ? LOCTEXT("DefaultBuildTargetName", "Project Default") : FText::Format(LOCTEXT("CurProjectDefaultBuildTargetName", "{0} (Project Default)"), FText::FromString(DefaultBuildTarget));
 		MenuBuilder.AddMenuEntry(
@@ -53,10 +57,10 @@ TSharedRef<SWidget> SCustomLaunchBuildTargetCombo::MakeBuildTargetSelectionWidge
 			FUIAction(
 				FExecuteAction::CreateSP(this, &SCustomLaunchBuildTargetCombo::SetBuildTargetName, FString()),
 				FCanExecuteAction(),
-				FGetActionCheckState::CreateLambda( [this]() { return SelectedBuildTarget.Get().IsEmpty() ? ECheckBoxState::Checked : ECheckBoxState::Unchecked; })
+				FGetActionCheckState::CreateLambda( [this]() { return SelectedBuildTargets.Get().IsEmpty() || SelectedBuildTargets.Get().Contains(FString()) ? ECheckBoxState::Checked : ECheckBoxState::Unchecked; })
 			), 
 			NAME_None, 
-			EUserInterfaceActionType::Check);
+			UIType);
 
 
 		if (SelectedProject.IsSet())
@@ -66,39 +70,44 @@ TSharedRef<SWidget> SCustomLaunchBuildTargetCombo::MakeBuildTargetSelectionWidge
 			const TArray<FTargetInfo>& BuildTargets = FDesktopPlatformModule::Get()->GetTargetsForProject(SelectedProject.Get());
 			for (const FTargetInfo& BuildTarget : BuildTargets)
 			{
-				if (BuildTarget.Type == EBuildTargetType::Program)
+				if (BuildTarget.Type == EBuildTargetType::Program || BuildTarget.Type == EBuildTargetType::Editor)
 				{
 					continue;
 				}
 
-				if (!SupportedTargetTypes.IsSet() || SupportedTargetTypes.Get().IsEmpty() || SupportedTargetTypes.Get().Contains(BuildTarget.Type))
-				{
-					
-					MenuBuilder.AddMenuEntry(
-						FText::FromString(BuildTarget.Name), 
-						FText::FromString(BuildTarget.Path), 
-						FSlateIcon(),
-						FUIAction(
-							FExecuteAction::CreateSP(this, &SCustomLaunchBuildTargetCombo::SetBuildTargetName, BuildTarget.Name),
-							FCanExecuteAction(),
-							FGetActionCheckState::CreateLambda( [this, BuildTargetName = BuildTarget.Name]() { return (SelectedBuildTarget.Get() == BuildTargetName) ? ECheckBoxState::Checked : ECheckBoxState::Unchecked; })
-						),
-						NAME_None,
-						EUserInterfaceActionType::Check);
-				}
-				else
-				{
-					MenuBuilder.AddMenuEntry(
-						FText::FromString(BuildTarget.Name), 
-						LOCTEXT("NotCompatible", "The selected platform does not support this build target type"), 
-						FSlateIcon(), 
-						FUIAction(
-							FExecuteAction::CreateSP(this, &SCustomLaunchBuildTargetCombo::SetBuildTargetName, BuildTarget.Name),
-							FCanExecuteAction::CreateLambda([]() { return false; })
-						));
-				}
+				MenuBuilder.AddMenuEntry(
+					FText::FromString(BuildTarget.Name), 
+					FText::FromString(BuildTarget.Path), 
+					FSlateIcon(),
+					FUIAction(
+						FExecuteAction::CreateSP(this, &SCustomLaunchBuildTargetCombo::SetBuildTargetName, BuildTarget.Name),
+						FCanExecuteAction(),
+						FGetActionCheckState::CreateSP(this, &SCustomLaunchBuildTargetCombo::GetBuildTargetCheckState, BuildTarget.Name)
+					),
+					NAME_None,
+					UIType);
 			}
 		}
+
+		if (SetAdvancedPlatformsOption.IsBound())
+		{
+			MenuBuilder.AddMenuSeparator();
+			MenuBuilder.AddMenuEntry(
+				LOCTEXT("AdvancedPlatformLabel", "Advanced build target selection..."),
+				LOCTEXT("AdvancedPlatformTip", "Select multiple platforms & build targets"),
+				FSlateIcon(),
+				FUIAction(
+					FExecuteAction::CreateLambda([this]()
+					{
+						ComboButton->SetIsOpen(false); // always close the popup menu when changing build target mode
+						SetAdvancedPlatformsOption.ExecuteIfBound(!GetAdvancedPlatformsOption.Get()); 
+					} ),
+					FCanExecuteAction(),
+					FGetActionCheckState::CreateLambda([this]() { return GetAdvancedPlatformsOption.Get() ? ECheckBoxState::Checked : ECheckBoxState::Unchecked; } )
+				),
+				NAME_None,
+				EUserInterfaceActionType::Check);
+		}
 	}
 
 	return MenuBuilder.MakeWidget();
@@ -108,27 +117,55 @@ END_SLATE_FUNCTION_BUILD_OPTIMIZATION
 
 FText SCustomLaunchBuildTargetCombo::GetBuildTargetName() const
 {
-	FString BuildTargetName = SelectedBuildTarget.Get();
+	TArray<FString> BuildTargetNames = SelectedBuildTargets.Get();
 
-	if (BuildTargetName.IsEmpty())
+	if (BuildTargetNames.Num() > 1)
 	{
-		FString DefaultBuildTarget = GetDefaultBuildTargetName();
-		if (DefaultBuildTarget.IsEmpty())
-		{
-			return LOCTEXT("DefaultBuildTargetName", "Project Default");
-		}
-		else
+		return FText::Format( LOCTEXT("MultipleTargets", "Multiple Build Targets ({0})"), FText::FromString(FString::Join(BuildTargetNames, TEXT(", "))));
+	}
+	else
+	{
+		FString BuildTargetName = (BuildTargetNames.Num() > 0) ? BuildTargetNames[0] : FString();
+
+		if (BuildTargetName.IsEmpty())
 		{
-			return FText::Format(LOCTEXT("CurProjectDefaultBuildTargetName", "{0} (Project Default)"), FText::FromString(DefaultBuildTarget));
+			FString DefaultBuildTarget = GetDefaultBuildTargetName();
+			if (DefaultBuildTarget.IsEmpty())
+			{
+				return LOCTEXT("DefaultBuildTargetName", "Project Default");
+			}
+			else
+			{
+				return FText::Format(LOCTEXT("CurProjectDefaultBuildTargetName", "{0} (Project Default)"), FText::FromString(DefaultBuildTarget));
+			}
 		}
-	}
 
-	return FText::FromString(BuildTargetName);
+		return FText::FromString(BuildTargetName);
+	}
 }
 
 void SCustomLaunchBuildTargetCombo::SetBuildTargetName(FString BuildTargetName)
 {
-	OnSelectionChanged.ExecuteIfBound(BuildTargetName);
+	TArray<FString> BuildTargets;
+	if (!BuildTargetName.IsEmpty())
+	{
+		if (GetAdvancedPlatformsOption.Get())
+		{
+			BuildTargets = SelectedBuildTargets.Get();
+		}
+
+		BuildTargets.Remove(FString());
+		if (BuildTargets.Contains(BuildTargetName))
+		{
+			BuildTargets.Remove(BuildTargetName);
+		}
+		else
+		{
+			BuildTargets.Add(BuildTargetName);
+		}
+	}
+
+	OnSelectionChanged.ExecuteIfBound(BuildTargets);
 }
 
 FString SCustomLaunchBuildTargetCombo::GetDefaultBuildTargetName() const
@@ -141,5 +178,22 @@ FString SCustomLaunchBuildTargetCombo::GetDefaultBuildTargetName() const
 	return FString();
 }
 
+ECheckBoxState SCustomLaunchBuildTargetCombo::GetBuildTargetCheckState( FString BuildTargetName ) const
+{
+	bool bResult = false;
+	const TArray<FString>& BuildTargets = SelectedBuildTargets.Get();
+
+	if (BuildTargetName.IsEmpty())
+	{
+		bResult = (BuildTargets.Num() == 0);
+	}
+	else
+	{
+		bResult = BuildTargets.Contains(BuildTargetName);
+	}
+
+	return bResult ? ECheckBoxState::Checked : ECheckBoxState::Unchecked;
+}
+
 
 #undef LOCTEXT_NAMESPACE
diff --git a/Engine/Plugins/Developer/ProjectLauncher/Source/ProjectLauncher/Private/Widgets/Shared/SCustomLaunchPlatformCombo.cpp b/Engine/Plugins/Developer/ProjectLauncher/Source/ProjectLauncher/Private/Widgets/Shared/SCustomLaunchPlatformCombo.cpp
--- a/Engine/Plugins/Developer/ProjectLauncher/Source/ProjectLauncher/Private/Widgets/Shared/SCustomLaunchPlatformCombo.cpp
+++ b/Engine/Plugins/Developer/ProjectLauncher/Source/ProjectLauncher/Private/Widgets/Shared/SCustomLaunchPlatformCombo.cpp
@@ -5,8 +5,10 @@
 #include "SlateOptMacros.h"
 #include "Widgets/SBoxPanel.h"
 #include "Widgets/Input/SComboBox.h"
+#include "Widgets/Input/SCheckBox.h"
 #include "Widgets/Text/STextBlock.h"
 #include "Widgets/Images/SImage.h"
+#include "Widgets/Layout/SSeparator.h"
 #include "PlatformInfo.h"
 
 
@@ -18,7 +20,9 @@ void SCustomLaunchPlatformCombo::Construct(const FArguments& InArgs)
 {
 	OnSelectionChanged = InArgs._OnSelectionChanged;
 	SelectedPlatforms = InArgs._SelectedPlatforms;
-	bBasicPlatformsOnly = InArgs._BasicPlatformsOnly;
+	GetAdvancedPlatformsOption = InArgs._GetAdvancedPlatformsOption;
+	SupportedBuildTargets = InArgs._SupportedBuildTargets;
+
 
 	FSlateFontInfo Font = InArgs._Font.IsSet() ? InArgs._Font.Get() : InArgs._TextStyle->Font;
 
@@ -28,15 +32,17 @@ void SCustomLaunchPlatformCombo::Construct(const FArguments& InArgs)
 		.OptionsSource(&PlatformsList)
 		.OnGenerateWidget( this, &SCustomLaunchPlatformCombo::OnGeneratePlatformListWidget )
 		.OnSelectionChanged( this, &SCustomLaunchPlatformCombo::OnPlatformSelectionChanged )
+		.MaxListHeight(1024)
 		[
 			SNew(SHorizontalBox)
 					
 			+SHorizontalBox::Slot()
 			.AutoWidth()
 			.VAlign(VAlign_Center)
-			.Padding(4,0)
+			.Padding(4,0, 4, 0)
 			[
 				SNew(SImage)
+				.Visibility_Lambda([this]() { return (SelectedPlatforms.Get().Num() <= 1) ? EVisibility::Visible : EVisibility::Collapsed; } )
 				.DesiredSizeOverride(FVector2D(16,16))
 				.Image(this, &SCustomLaunchPlatformCombo::GetSelectedPlatformBrush)
 			]
@@ -44,7 +50,6 @@ void SCustomLaunchPlatformCombo::Construct(const FArguments& InArgs)
 			+SHorizontalBox::Slot()
 			.AutoWidth()
 			.VAlign(VAlign_Center)
-			.Padding(4,0)
 			[
 				SNew(STextBlock)
 				.Text(this, &SCustomLaunchPlatformCombo::GetSelectedPlatformName)
@@ -53,50 +58,59 @@ void SCustomLaunchPlatformCombo::Construct(const FArguments& InArgs)
 		]
 	];
 
-	for (const PlatformInfo::FTargetPlatformInfo* PlatformInfo : PlatformInfo::GetPlatformInfoArray())
-	{
-		if( bBasicPlatformsOnly && PlatformInfo->PlatformType != EBuildTargetType::Game)
-		{
-			continue;
-		}
-
-		PlatformsList.Add(MakeShared<FString>(PlatformInfo->Name.ToString()));
-	}
-
+	RefreshPlatformsList();
 }
 END_SLATE_FUNCTION_BUILD_OPTIMIZATION
 
 
 
 BEGIN_SLATE_FUNCTION_BUILD_OPTIMIZATION
-TSharedRef<SWidget> SCustomLaunchPlatformCombo::OnGeneratePlatformListWidget( TSharedPtr<FString> Platform ) const
+TSharedRef<SWidget> SCustomLaunchPlatformCombo::OnGeneratePlatformListWidget( TSharedPtr<FString> Platform )
 {
-	const PlatformInfo::FTargetPlatformInfo* PlatformInfo = PlatformInfo::FindPlatformInfo(FName(*Platform));
+	const PlatformInfo::FTargetPlatformInfo* PlatformInfo = (!Platform.IsValid() || Platform->IsEmpty()) ? nullptr : PlatformInfo::FindPlatformInfo(FName(*Platform));
 	if (PlatformInfo != nullptr)
 	{
+		bool bIsGroupStart = BuildTargetGroupStarts.Contains(*Platform);
+		int32 GroupPadding = bIsGroupStart ? 16 : 0;
 		int32 Indent = PlatformInfo->IsVanilla() ? 0 : 16;
 
-		return SNew(SHorizontalBox)
-
-			+SHorizontalBox::Slot()
-			.AutoWidth()
-			.VAlign(VAlign_Center)
-			.Padding(4+Indent,4,4,4)
-			[
-				SNew(SImage)
-				.DesiredSizeOverride(FVector2D(16,16))
-				.Image(FAppStyle::GetBrush(PlatformInfo->GetIconStyleName(EPlatformIconSize::Normal)))
-			]
-
-			+SHorizontalBox::Slot()
-			.AutoWidth()
-			.VAlign(VAlign_Center)
-			.Padding(4)
-			[
-				SNew(STextBlock)
-				.Text(PlatformInfo->DisplayName)
+		return SNew(SBox)
+			.Padding(0,GroupPadding,0,0)
+			[			
+				SNew(SHorizontalBox)
+
+				+SHorizontalBox::Slot()
+				.AutoWidth()
+				.VAlign(VAlign_Center)
+				[
+					SNew(SCheckBox)
+					.Visibility_Lambda( [this]() { return GetAdvancedPlatformsOption.Get() ? EVisibility::Visible : EVisibility::Collapsed; } )
+					.IsChecked( this, &SCustomLaunchPlatformCombo::GetPlatformCheckState, Platform )
+					.OnCheckStateChanged( this, &SCustomLaunchPlatformCombo::SetPlatformCheckState, Platform )
+				]
+
+				+SHorizontalBox::Slot()
+				.AutoWidth()
+				.VAlign(VAlign_Center)
+				.Padding(4+Indent,2,4,2)
+				[
+					SNew(SImage)
+					.DesiredSizeOverride(FVector2D(16,16))
+					.Image(FAppStyle::GetBrush(PlatformInfo->GetIconStyleName(EPlatformIconSize::Normal)))
+				]
+
+				+SHorizontalBox::Slot()
+				.AutoWidth()
+				.VAlign(VAlign_Center)
+				.Padding(2,2,2,2)
+				[
+					SNew(STextBlock)
+					.Text(PlatformInfo->DisplayName)
+					.ColorAndOpacity_Lambda( [this, PlatformInfo] { return SCustomLaunchPlatformCombo::GetPlatformTextColor(PlatformInfo); } )
+				]
 			]
 		;
+
 	}
 	else
 	{
@@ -109,14 +123,18 @@ END_SLATE_FUNCTION_BUILD_OPTIMIZATION
 
 void SCustomLaunchPlatformCombo::OnPlatformSelectionChanged( TSharedPtr<FString> Platform, ESelectInfo::Type InSelectInfo )
 {
+	if (!Platform.IsValid())
+	{
+		return;
+	}
+
 	TArray<FString> Platforms;
-	if (Platform.IsValid() && !Platform->IsEmpty())
+	if (!Platform->IsEmpty())
 	{
 		Platforms.Add(*Platform);
 	}
 
-	SelectedPlatforms.Set(Platforms);
-	OnSelectionChanged.ExecuteIfBound(SelectedPlatforms.Get());
+	OnSelectionChanged.ExecuteIfBound(Platforms);
 }
 
 
@@ -133,7 +151,7 @@ const FSlateBrush* SCustomLaunchPlatformCombo::GetSelectedPlatformBrush() const
 			return FAppStyle::GetBrush(PlatformInfo->GetIconStyleName(EPlatformIconSize::Normal));
 		}
 	}
-	else if (Platforms.Num() > 1)
+	else if (Platforms.Num() == 0)
 	{
 		return FAppStyle::Get().GetBrush("Icons.WarningWithColor");
 	}
@@ -157,11 +175,142 @@ FText SCustomLaunchPlatformCombo::GetSelectedPlatformName() const
 	}
 	else if (Platforms.Num() > 1)
 	{
-		return LOCTEXT("TooManyPlatforms", "Multiple platforms (unsupported)");
+		return FText::Format( LOCTEXT("MultiplePlatforms", "Multiple Platforms ({0})"), FText::FromString(FString::Join(Platforms, TEXT(", "))));
 	}
 
 	return LOCTEXT("NoPlatform", "(no platform)");
 }
 
 
+ECheckBoxState SCustomLaunchPlatformCombo::GetPlatformCheckState( const TSharedPtr<FString> Platform ) const
+{
+	if (!Platform.IsValid() || Platform->IsEmpty())
+	{
+		return ECheckBoxState::Unchecked;
+	}
+
+	const TArray<FString>& Platforms = SelectedPlatforms.Get();
+	return Platforms.Contains(*Platform) ? ECheckBoxState::Checked : ECheckBoxState::Unchecked;
+}
+
+void SCustomLaunchPlatformCombo::SetPlatformCheckState( ECheckBoxState CheckState, TSharedPtr<FString> Platform )
+{
+	if (!Platform.IsValid() || Platform->IsEmpty())
+	{
+		return;
+	}
+
+	TArray<FString> Platforms = SelectedPlatforms.Get();
+	if (CheckState == ECheckBoxState::Checked)
+	{
+		Platforms.AddUnique(*Platform);
+	}
+	else
+	{
+		Platforms.Remove(*Platform);
+	}
+
+	OnSelectionChanged.ExecuteIfBound(Platforms);
+}
+
+FSlateColor SCustomLaunchPlatformCombo::GetPlatformTextColor( const PlatformInfo::FTargetPlatformInfo* PlatformInfo )
+{
+	TSet<EBuildTargetType> BuildTargets = SupportedBuildTargets.Get();
+	if (PlatformInfo != nullptr && BuildTargets.Num() > 0)
+	{
+		bool bHasSupportedBuildTarget;
+
+		if (GetAdvancedPlatformsOption.Get())
+		{
+			bHasSupportedBuildTarget = BuildTargets.Contains(PlatformInfo->PlatformType);
+		}
+		else
+		{
+			bHasSupportedBuildTarget = PlatformInfo->VanillaInfo->Flavors.ContainsByPredicate( [BuildTargets]( const PlatformInfo::FTargetPlatformInfo* Info )
+			{
+				return BuildTargets.Contains(Info->PlatformType);
+			});
+		}
+
+		if (!bHasSupportedBuildTarget)
+		{
+			return FSlateColor::UseSubduedForeground();
+		}
+	}
+
+
+	return FSlateColor::UseForeground();
+}
+
+void SCustomLaunchPlatformCombo::RefreshPlatformsList()
+{
+	PlatformsList.Reset();
+	BuildTargetGroupStarts.Reset();
+
+	if (GetAdvancedPlatformsOption.Get())
+	{
+		// sort by Client/Server, then by IniPlatformName, then by Flavor
+		TArray<PlatformInfo::FTargetPlatformInfo*> SortedPlatformInfos = PlatformInfo::GetPlatformInfoArray();
+		SortedPlatformInfos.Sort( [](const PlatformInfo::FTargetPlatformInfo& A, const PlatformInfo::FTargetPlatformInfo& B)
+		{
+			if (A.PlatformType != B.PlatformType)
+			{
+				return (int)A.PlatformType < (int)B.PlatformType;
+			}
+
+			if (A.IniPlatformName != B.IniPlatformName)
+			{
+				return A.IniPlatformName.ToString() < B.IniPlatformName.ToString();
+			}
+
+			if (A.IsVanilla() != B.IsVanilla())
+			{
+				return A.IsVanilla();
+			}
+
+			return A.Name.ToString() < B.Name.ToString(); // fallback - should not get here
+		});
+
+
+		EBuildTargetType PrevPlatformType = (SortedPlatformInfos.Num() > 0) ? SortedPlatformInfos[0]->PlatformType : EBuildTargetType::Unknown;
+		for (const PlatformInfo::FTargetPlatformInfo* PlatformInfo : SortedPlatformInfos)
+		{
+			if (PlatformInfo->PlatformType == EBuildTargetType::Program || PlatformInfo->PlatformType == EBuildTargetType::Editor)
+			{
+				continue;
+			}
+
+
+			FString PlatformName = PlatformInfo->Name.ToString();
+
+			PlatformsList.Add(MakeShared<FString>(PlatformName));
+
+			// take note of which platform starts a new build type group so we can add some separation between them in the list
+			if (PlatformInfo->PlatformType != PrevPlatformType)
+			{
+				BuildTargetGroupStarts.Add(PlatformName);
+				PrevPlatformType = PlatformInfo->PlatformType;
+			}
+		}
+
+	}
+	else
+	{
+		for (const PlatformInfo::FTargetPlatformInfo* PlatformInfo : PlatformInfo::GetPlatformInfoArray())
+		{
+			if( PlatformInfo->PlatformType != EBuildTargetType::Game)
+			{
+				continue;
+			}
+
+			PlatformsList.Add(MakeShared<FString>(PlatformInfo->Name.ToString()));
+		}
+	}
+
+	PlatformsComboBox->RefreshOptions();
+}
+
+
+
+
 #undef LOCTEXT_NAMESPACE
diff --git a/Engine/Plugins/Developer/ProjectLauncher/Source/ProjectLauncher/Public/Extension/LaunchExtension.h b/Engine/Plugins/Developer/ProjectLauncher/Source/ProjectLauncher/Public/Extension/LaunchExtension.h
--- a/Engine/Plugins/Developer/ProjectLauncher/Source/ProjectLauncher/Public/Extension/LaunchExtension.h
+++ b/Engine/Plugins/Developer/ProjectLauncher/Source/ProjectLauncher/Public/Extension/LaunchExtension.h
@@ -77,6 +77,11 @@ namespace ProjectLauncher
 		 * Advanced hook to allow any advanced modification of the command line when our profile is launched
 		 */
 		virtual void CustomizeLaunchCommandLine( FString& InOutCommandLine ) {};
+		
+		/**
+		 * Advanced hook to allow any advanced modification of the command line when our profile is launched for the given build target type
+		 */
+		virtual void CustomizeTargetLaunchCommandLine( FString& InOutCommandLine, EBuildTargetType InBuildTargetType ) {};
 
 		/**
 		 * Advanced hook to allow any advanced modification of the command line when an automated test is launched
@@ -237,7 +242,7 @@ namespace ProjectLauncher
 		inline TSharedRef<FModel> GetModel() const { return Model; }
 
 		/**
-		 * Get the name of the platform to launch on
+		 * Get the name of the platform to launch on. this is the highest priority platform/target in the list of platforms
 		 */
 		UE_API FName GetLaunchPlatformName() const;
 
diff --git a/Engine/Plugins/Developer/ProjectLauncher/Source/ProjectLauncher/Public/Model/ProjectLauncherModel.h b/Engine/Plugins/Developer/ProjectLauncher/Source/ProjectLauncher/Public/Model/ProjectLauncherModel.h
--- a/Engine/Plugins/Developer/ProjectLauncher/Source/ProjectLauncher/Public/Model/ProjectLauncherModel.h
+++ b/Engine/Plugins/Developer/ProjectLauncher/Source/ProjectLauncher/Public/Model/ProjectLauncherModel.h
@@ -25,6 +25,8 @@ struct FTargetInfo;
 
 namespace ProjectLauncher
 {
+	extern const FText ProjectLauncherText;
+
 	/* Defines how content is stored and deployed for a profile */
 	enum class EContentScheme : uint8
 	{
@@ -112,8 +114,6 @@ namespace ProjectLauncher
 	DECLARE_MULTICAST_DELEGATE(FOnSelectedPlatformChanged)
 
 
-	extern bool bUseFriendlyBuildTargetSelection;
-
 	/*
 	 * This class contains a selection of helper functions and the current profile state
 	 */
@@ -215,16 +215,22 @@ namespace ProjectLauncher
 		}
 
 		/* Get platform info for a profile or platform name */
+		UE_DEPRECATED(5.8, "Use GetPlatformInfos")
 		static UE_API const PlatformInfo::FTargetPlatformInfo* GetPlatformInfo(const ILauncherProfilePtr& Profile);
 		static UE_API const PlatformInfo::FTargetPlatformInfo* GetPlatformInfo(FName PlatformName, const FTargetInfo& BuildTargetInfo);
+		static UE_API TArray<const PlatformInfo::FTargetPlatformInfo*> GetPlatformInfos(const ILauncherProfilePtr& Profile);
 
 		/* Returns true if profile or platform name refers to host platform */
+		UE_DEPRECATED(5.8, "No longer in use. IsUsingRemotePlatform may be of use")
 		static UE_API bool IsHostPlatform(const ILauncherProfilePtr& Profile);
 		static UE_API bool IsHostPlatform(FName PlatformName);
+		static UE_API bool IsUsingRemotePlatform(const ILauncherProfilePtr& Profile);
 
 		/* Get build target info for a project path or profile */
 		UE_API FTargetInfo GetBuildTargetInfo(FString BuildTargetName, const FString& ProjectPath);
+		UE_DEPRECATED(5.8, "Use GetBuildTargetInfos")
 		UE_API FTargetInfo GetBuildTargetInfo(const ILauncherProfileRef& Profile);
+		UE_API TArray<FTargetInfo> GetBuildTargetInfos(const ILauncherProfileRef& Profile);
 
 		/* Get a device proxy for a profile */
 		static UE_API TSharedPtr<ITargetDeviceProxy> GetDeviceProxy(const ILauncherProfileRef& Profile);
@@ -276,7 +282,7 @@ namespace ProjectLauncher
 		UE_API void ReadProjectSettingsFromConfig(FConfigCacheIni& InConfig, const FString& InProjectPath, FProjectSettings& OutResult);
 
 		/* Modify launch command lines to inject extension variables */
-		UE_API void OnModifyLaunchCommandLine(const ILauncherProfileRef& InProfile, FString& InOutCommandLine);
+		UE_API void OnModifyLaunchCommandLine( const ILauncherProfileRef& InProfile, EBuildTargetType InBuildTargetType, FString& InOutCommandLine );
 		UE_API void OnModifyAutomatedTestCommandLine(const ILauncherProfileAutomatedTestRef& InAutomatedTest, const ILauncherProfileRef& InProfile, FString& InOutCommandLine);
 
 		/* Get map paths for a project, caching results */
diff --git a/Engine/Plugins/Developer/ProjectLauncher/Source/ProjectLauncher/Public/ProfileTree/GenericProfileTreeBuilder.h b/Engine/Plugins/Developer/ProjectLauncher/Source/ProjectLauncher/Public/ProfileTree/GenericProfileTreeBuilder.h
--- a/Engine/Plugins/Developer/ProjectLauncher/Source/ProjectLauncher/Public/ProfileTree/GenericProfileTreeBuilder.h
+++ b/Engine/Plugins/Developer/ProjectLauncher/Source/ProjectLauncher/Public/ProfileTree/GenericProfileTreeBuilder.h
@@ -4,10 +4,14 @@
 
 #include "ProfileTree/ILaunchProfileTreeBuilder.h"
 #include "Model/ProjectLauncherModel.h"
+#include "Containers/Map.h"
+#include "Containers/Set.h"
 
 #define UE_API PROJECTLAUNCHER_API
 
 class SSearchableComboBox;
+class SCustomLaunchPlatformCombo;
+class SCustomLaunchBuildTargetCombo;
 class SCustomLaunchDeviceListView;
 namespace ESelectInfo { enum Type : int; }
 
@@ -96,9 +100,8 @@ namespace ProjectLauncher
 		UE_API void SetProjectName(FString ProjectPath);
 		UE_API bool HasProject() const;
 
-		UE_API FString GetBuildTarget() const;
-		UE_API void SetBuildTarget(FString BuildTarget);
-		UE_API TArray<EBuildTargetType> GetSupportedBuildTargetTypes() const;
+		UE_API TArray<FString> GetBuildTargets() const;
+		UE_API void SetBuildTargets( TArray<FString> BuildTargets);
 
 		UE_API void SetBuildConfiguration(EBuildConfiguration BuildConfiguration);
 		UE_API EBuildConfiguration GetBuildConfiguration() const;
@@ -106,8 +109,8 @@ namespace ProjectLauncher
 		UE_API void SetContentScheme(EContentScheme ContentScheme);
 		UE_API bool IsContentSchemeAvailable(EContentScheme, FText& OutReason) const;
 
-		UE_API FString GetCommandLine() const;
-		UE_API void SetCommandLine( const FString& NewCommandLine );
+		UE_API FString GetCommandLine(EBuildTargetType BuildTarget) const;
+		UE_API void SetCommandLine( const FString& NewCommandLine, EBuildTargetType BuildTarget );
 
 		UE_API void SetSelectedPlatforms( TArray<FString> SelectedPlatforms );
 		UE_API TArray<FString> GetSelectedPlatforms() const;
@@ -131,8 +134,8 @@ namespace ProjectLauncher
 		UE_API void SetForceBuild( bool bForceBuild );
 		UE_API bool GetForceBuild( ILauncherProfilePtr InProfile = nullptr ) const;
 
-		UE_API void SetArchitecture( FString Architecture );
-		UE_API FString GetArchitecture() const;
+		UE_API void SetArchitecture( FString Architecture, EBuildTargetType TargetType );
+		UE_API FString GetArchitecture(EBuildTargetType TargetType) const;
 		UE_API FText GetArchitectureDisplayName( FString Architecture );
 
 		UE_API void SetDeployToDevice( bool bDeployToDevice );
@@ -145,6 +148,9 @@ namespace ProjectLauncher
 		UE_API TSharedPtr<FString> GetInitialMap() const;
 
 
+		UE_API void SetHasAdvancedPlatformTargets( bool bEnable );
+		UE_API bool HasAdvancedPlatformTargets() const;
+
 
 	protected:
 
@@ -153,9 +159,12 @@ namespace ProjectLauncher
 		UE_API void RefreshContentScheme();
 		UE_API void CacheStartupMapList() const;
 		UE_API void CacheArchitectures();
+		UE_API void CacheBuildTargets();
+		UE_API void ValidateArchitectures();
 		UE_API TSharedRef<SWidget> OnGenerateComboWidget( TSharedPtr<FString> InComboString );
 		FText GetZenSnapshotText() const;
 
+
 		FLaunchProfileTreeDataRef TreeData;
 		const ILauncherProfileRef Profile;
 		const ILauncherProfileRef DefaultProfile;
@@ -164,13 +173,16 @@ namespace ProjectLauncher
 		const TSharedRef<FModel> Model;
 
 		EContentScheme ContentScheme;
+		bool bHasAdvancedPlatformTargets;
 		bool bShouldCook;
 		mutable bool bStartupMapCacheDirty = true;
 		mutable TArray<TSharedPtr<FString>> CachedStartupMaps;
+		TSharedPtr<SCustomLaunchPlatformCombo> PlatformCombo;
+		TSharedPtr<SCustomLaunchBuildTargetCombo> BuildTargetCombo;
 		TSharedPtr<SSearchableComboBox> InitalMapCombo;
 		TSharedPtr<SCustomLaunchDeviceListView> DeployDeviceListView;
-		TArray<FString> CachedArchitectures;
-		EBuildTargetType CachedBuildTargetType = EBuildTargetType::Game;
+		TMap<EBuildTargetType, TArray<FString>> CachedBuildTargetArchitectures;
+		TSet<EBuildTargetType> CachedBuildTargetTypes;
 	};
 }
 
diff --git a/Engine/Plugins/Developer/ProjectLauncher/Source/ProjectLauncher/Public/Widgets/Shared/SCustomLaunchBuildTargetCombo.h b/Engine/Plugins/Developer/ProjectLauncher/Source/ProjectLauncher/Public/Widgets/Shared/SCustomLaunchBuildTargetCombo.h
--- a/Engine/Plugins/Developer/ProjectLauncher/Source/ProjectLauncher/Public/Widgets/Shared/SCustomLaunchBuildTargetCombo.h
+++ b/Engine/Plugins/Developer/ProjectLauncher/Source/ProjectLauncher/Public/Widgets/Shared/SCustomLaunchBuildTargetCombo.h
@@ -7,6 +7,7 @@
 #include "Styling/SlateTypes.h"
 #include "Styling/CoreStyle.h"
 #include "Styling/SlateWidgetStyleAsset.h"
+#include "Framework/SlateDelegates.h"
 #include "Fonts/SlateFontInfo.h"
 #include "Model/ProjectLauncherModel.h"
 
@@ -16,16 +17,17 @@ class SCustomLaunchBuildTargetCombo
 	: public SCompoundWidget
 {
 public:
-	DECLARE_DELEGATE_OneParam(FOnSelectionChanged, FString );
+	DECLARE_DELEGATE_OneParam(FOnSelectionChanged, TArray<FString> );
 
 	SLATE_BEGIN_ARGS(SCustomLaunchBuildTargetCombo)
 		: _TextStyle( &FCoreStyle::Get().GetWidgetStyle<FTextBlockStyle>( "NormalText" ) )
 		{}
 		SLATE_EVENT(FOnSelectionChanged, OnSelectionChanged);
 		SLATE_ATTRIBUTE(FString, SelectedProject);
-		SLATE_ATTRIBUTE(FString, SelectedBuildTarget);
-		SLATE_ATTRIBUTE(TArray<EBuildTargetType>, SupportedTargetTypes);
+		SLATE_ATTRIBUTE(TArray<FString>, SelectedBuildTargets);
 		SLATE_ATTRIBUTE(FSlateFontInfo, Font)
+		SLATE_ATTRIBUTE(bool, GetAdvancedPlatformsOption)
+		SLATE_EVENT(FOnBooleanValueChanged, SetAdvancedPlatformsOption)
 		SLATE_STYLE_ARGUMENT( FTextBlockStyle, TextStyle )
 	SLATE_END_ARGS()
 
@@ -34,13 +36,17 @@ public:
 
 protected:
 	UE_API FString GetDefaultBuildTargetName() const;
+	UE_API ECheckBoxState GetBuildTargetCheckState( FString BuildTargetName ) const;
 
-	TAttribute<TArray<EBuildTargetType>> SupportedTargetTypes;
 	TAttribute<FString> SelectedProject;
-	TAttribute<FString> SelectedBuildTarget;
+	TAttribute<TArray<FString>> SelectedBuildTargets;
 	FOnSelectionChanged OnSelectionChanged;
 	bool bShowAnyProjectOption;
+	TAttribute<bool> GetAdvancedPlatformsOption;
+	FOnBooleanValueChanged SetAdvancedPlatformsOption;
+
 	TSharedPtr<ProjectLauncher::FModel> Model;
+	TSharedPtr<class SComboButton> ComboButton;
 
 	UE_API TSharedRef<SWidget> MakeBuildTargetSelectionWidget();
 
diff --git a/Engine/Plugins/Developer/ProjectLauncher/Source/ProjectLauncher/Public/Widgets/Shared/SCustomLaunchPlatformCombo.h b/Engine/Plugins/Developer/ProjectLauncher/Source/ProjectLauncher/Public/Widgets/Shared/SCustomLaunchPlatformCombo.h
--- a/Engine/Plugins/Developer/ProjectLauncher/Source/ProjectLauncher/Public/Widgets/Shared/SCustomLaunchPlatformCombo.h
+++ b/Engine/Plugins/Developer/ProjectLauncher/Source/ProjectLauncher/Public/Widgets/Shared/SCustomLaunchPlatformCombo.h
@@ -3,15 +3,22 @@
 #pragma once
 
 #include "CoreMinimal.h"
+#include "HAL/PlatformMisc.h"
 #include "Widgets/SCompoundWidget.h"
 #include "Styling/CoreStyle.h"
 #include "Styling/SlateWidgetStyleAsset.h"
 #include "Fonts/SlateFontInfo.h"
+#include "Framework/SlateDelegates.h"
 
 #define UE_API PROJECTLAUNCHER_API
 
 template<typename ItemType> class SComboBox;
 
+namespace PlatformInfo
+{
+	struct FTargetPlatformInfo;
+}
+
 
 class SCustomLaunchPlatformCombo
 	: public SCompoundWidget
@@ -24,23 +31,32 @@ public:
 		{}
 		SLATE_EVENT(FOnSelectionChanged, OnSelectionChanged);
 		SLATE_ATTRIBUTE(TArray<FString>, SelectedPlatforms);
-		SLATE_ARGUMENT(bool, BasicPlatformsOnly)
+		SLATE_ATTRIBUTE(bool, GetAdvancedPlatformsOption)
+		SLATE_ATTRIBUTE(TSet<EBuildTargetType>, SupportedBuildTargets)
 		SLATE_ATTRIBUTE(FSlateFontInfo, Font)
 		SLATE_STYLE_ARGUMENT( FTextBlockStyle, TextStyle )
 	SLATE_END_ARGS()
 
 public:
 	UE_API void Construct(	const FArguments& InArgs);
+	UE_API void RefreshPlatformsList();
 
 protected:
 	TAttribute<TArray<FString>> SelectedPlatforms;
 	FOnSelectionChanged OnSelectionChanged;
-	bool bBasicPlatformsOnly;
+	TAttribute<bool> GetAdvancedPlatformsOption;
+	TAttribute<TSet<EBuildTargetType>> SupportedBuildTargets;
+	TSet<FString> BuildTargetGroupStarts;
 
-	UE_API TSharedRef<SWidget> OnGeneratePlatformListWidget( TSharedPtr<FString> Platform ) const;
+	UE_API TSharedRef<SWidget> OnGeneratePlatformListWidget( TSharedPtr<FString> Platform );
 	UE_API void OnPlatformSelectionChanged( TSharedPtr<FString> Platform, ESelectInfo::Type InSelectInfo );
 	UE_API const FSlateBrush* GetSelectedPlatformBrush() const;
 	UE_API FText GetSelectedPlatformName() const;
+	UE_API ECheckBoxState GetPlatformCheckState( const TSharedPtr<FString> Platform ) const;
+	UE_API void SetPlatformCheckState( ECheckBoxState CheckState, TSharedPtr<FString> Platform );
+	UE_API FSlateColor GetPlatformTextColor( const PlatformInfo::FTargetPlatformInfo* PlatformInfo );
+
+
 	TArray<TSharedPtr<FString>> PlatformsList;
 	TSharedPtr<SComboBox<TSharedPtr<FString>>> PlatformsComboBox;
 };
diff --git a/Engine/Source/Developer/LauncherServices/Private/Launcher/LauncherWorker.cpp b/Engine/Source/Developer/LauncherServices/Private/Launcher/LauncherWorker.cpp
--- a/Engine/Source/Developer/LauncherServices/Private/Launcher/LauncherWorker.cpp
+++ b/Engine/Source/Developer/LauncherServices/Private/Launcher/LauncherWorker.cpp
@@ -417,6 +417,7 @@ FString FLauncherWorker::CreateUATCommand( const ILauncherProfileRef& InProfile,
 	FString OptionalParams = TEXT("");
 	TSet<FString> OptionalTargetPlatforms;
 	TSet<FString> OptionalCookFlavors;
+	TSet<EBuildTargetType> BuildTargets;
 	
 	bool bUATClosesAfterLaunch = false;
 	for (int32 PlatformIndex = 0; PlatformIndex < InPlatforms.Num(); ++PlatformIndex)
@@ -459,6 +460,8 @@ FString FLauncherWorker::CreateUATCommand( const ILauncherProfileRef& InProfile,
 				OptionalParams += OptionalUATCommandLine;
 			}
 			bUATClosesAfterLaunch |= PlatformInfo->DataDrivenPlatformInfo->bUATClosesAfterLaunch;
+
+			BuildTargets.Add(PlatformInfo->PlatformType);
 		}
 	}
 
@@ -502,9 +505,9 @@ FString FLauncherWorker::CreateUATCommand( const ILauncherProfileRef& InProfile,
 		UATCommand += (TEXT(" -cookflavor=") + Join(OptionalCookFlavors, TEXT("+")));
 	}
 
-	if (InProfile->GetBuildTarget().Len() > 0)
+	if (InProfile->GetBuildTargets().Num() > 0 && !InProfile->GetBuildTargets().Contains(FString()))
 	{
-		UATCommand += TEXT(" -target=") + InProfile->GetBuildTarget();
+		UATCommand += TEXT(" -target=") + FString::Join(InProfile->GetBuildTargets(), TEXT("+"));
 	}
 
 	if (InProfile->IsDeviceASimulator())
@@ -603,7 +606,7 @@ FString FLauncherWorker::CreateUATCommand( const ILauncherProfileRef& InProfile,
 	FString SessionCommands;
 	if (!bForTurnkeyCustomBuild)
 	{
-		FString SessionName = InProfile->GetName().Replace(TEXT("\'"), TEXT("_")).Replace(TEXT("\'"), TEXT("_"));
+		FString SessionName = InProfile->GetName().Replace(TEXT("\'"), TEXT("_")).Replace(TEXT("\'"), TEXT("_")).Replace(TEXT("&"), TEXT("_"));
 		FString SessionOwner = FString(FPlatformProcess::UserName(false)).Replace(TEXT("\'"), TEXT("_")).Replace(TEXT("\'"), TEXT("_"));;
 		SessionCommands = FString::Printf(TEXT("-SessionId=%s -SessionOwner='%s' -SessionName='%s' "),
 		*SessionId.ToString(),
@@ -613,7 +616,7 @@ FString FLauncherWorker::CreateUATCommand( const ILauncherProfileRef& InProfile,
 
 	// allow external items to adjust the command line
 	FString ProfileAdditionalCommandLineParameters = InProfile->GetAdditionalCommandLineParameters();
-	LauncherProfileManager->OnPostProcessLaunchCommandLine().Broadcast(InProfile, ProfileAdditionalCommandLineParameters);
+	LauncherProfileManager->OnPostProcessLaunchCommandLine().Broadcast(InProfile, EBuildTargetType::Unknown, ProfileAdditionalCommandLineParameters);
 
 	FString AdditionalCommandLine;
 	TStringBuilder<64> InheritableGameOptions;
@@ -634,6 +637,34 @@ FString FLauncherWorker::CreateUATCommand( const ILauncherProfileRef& InProfile,
 		AdditionalCommandLine = FString::Printf(TEXT(" -addcmdline=\"%s\""), *AddCmdLine);
 	}
 
+	for (EBuildTargetType BuildTarget : BuildTargets)
+	{
+		FString BuildTargetCmdLine = InProfile->GetAdditionalTargetCommandLineParameters(BuildTarget);
+		LauncherProfileManager->OnPostProcessLaunchCommandLine().Broadcast(InProfile, BuildTarget, BuildTargetCmdLine);
+		if (BuildTargetCmdLine.TrimStartAndEnd().IsEmpty())
+		{
+			continue;		
+		}
+		BuildTargetCmdLine.ReplaceInline(TEXT("\""), TEXT("\\\""));
+
+		switch (BuildTarget)
+		{
+			case EBuildTargetType::Server:
+			{
+				AdditionalCommandLine += FString::Printf(TEXT(" -servercmdline=\"%s\""), *BuildTargetCmdLine);
+			}
+			break;
+
+			case EBuildTargetType::Client:
+			case EBuildTargetType::Game:
+			{
+				AdditionalCommandLine += FString::Printf(TEXT(" -clientcmdline=\"%s\""), *BuildTargetCmdLine);
+			}
+			break;
+		}
+	}
+
+
 
 	// map list
 	FString MapList;
diff --git a/Engine/Source/Developer/LauncherServices/Private/Profiles/LauncherProfile.cpp b/Engine/Source/Developer/LauncherServices/Private/Profiles/LauncherProfile.cpp
--- a/Engine/Source/Developer/LauncherServices/Private/Profiles/LauncherProfile.cpp
+++ b/Engine/Source/Developer/LauncherServices/Private/Profiles/LauncherProfile.cpp
@@ -111,7 +111,7 @@ void FLauncherProfile::QueryZenSnapshotsForProject()
 					}
 				}
 
-				PlatformConfigToBuildMap.Add(Pair.Key, ClosestBuild);
+				PlatformConfigToBuildMap.Add(Pair.Key.ToLower(), ClosestBuild);
 			}
 		}
 	);
@@ -123,13 +123,13 @@ int32 FLauncherProfile::GetZenSnapshot()
 	if (PerProjectBuilds.Contains(ProjectName))
 	{
 		TArray<FString> Platforms = GetCookedPlatforms();
-		if (Platforms.Num() > 0)
+		for (const FString& Platform : Platforms)
 		{
-			FString Platform = Platforms[0].ToLower();
+			FString LowerPlatform = Platform.ToLower();
 			FPlatformToBuildsMap& Builds = PerProjectBuilds[ProjectName];
-			if (Builds.Contains(Platform))
+			if (Builds.Contains(LowerPlatform))
 			{
-				return Builds[Platform];
+				return Builds[LowerPlatform];
 			}
 		}
 	}
diff --git a/Engine/Source/Developer/LauncherServices/Private/Profiles/LauncherProfile.h b/Engine/Source/Developer/LauncherServices/Private/Profiles/LauncherProfile.h
--- a/Engine/Source/Developer/LauncherServices/Private/Profiles/LauncherProfile.h
+++ b/Engine/Source/Developer/LauncherServices/Private/Profiles/LauncherProfile.h
@@ -66,6 +66,7 @@ enum ELauncherVersion
 	LAUNCHERSERVICES_AUTOMATEDTESTS = 41,
 	LAUNCHERSERVICES_FASTITERATE = 42,
 	LAUNCHERSERVICES_REMOVEDINCREMENTALCOOKMODE = 43,
+	LAUNCHERSERVICES_ADDMULTIPLEBUILDTARGETS = 44,
 	//ADD NEW STUFF HERE
 
 
@@ -425,7 +426,9 @@ public:
 	{
 		if (BuildTargetSpecified)
 		{
-			return BuildTargetName;
+			UE_CLOG(BuildTargets.Num() > 0, LogTemp, Warning, TEXT("GetBuildTargets is only returning the first build target - legacy project launcher?"));
+
+			return (BuildTargets.Num() > 0) ? BuildTargets[0] : FString();
 		}
 		else
 		{
@@ -879,6 +882,14 @@ public:
 		return AdditionalCommandLineParameters;
 	}
 
+	virtual FString GetAdditionalTargetCommandLineParameters( EBuildTargetType BuildTargetType ) const override
+	{
+		check( BuildTargetType != EBuildTargetType::Unknown);
+
+		const FString* CommandLinePtr = AdditionalTargetCommandLineParameters.Find(BuildTargetType);
+		return CommandLinePtr ? *CommandLinePtr : FString();
+	}
+
 	virtual bool IsCookingIncrementally( ) const override
 	{
 		if ( CookMode != ELauncherProfileCookModes::DoNotCook )
@@ -1180,6 +1191,7 @@ public:
 			Archive << OriginalReleaseVersionName;
 		}
 
+		FString BuildTargetName;
 		if (Version >= LAUNCHERSERVICES_ADDBUILDTARGETNAME)
 		{
 			Archive << BuildTargetSpecified;
@@ -1244,6 +1256,17 @@ public:
 			Archive << bWantAutomatedTestBuild;
 		}
 
+		if (Version >= LAUNCHERSERVICES_ADDMULTIPLEBUILDTARGETS)
+		{
+			Archive << BuildTargets;
+			if (Archive.IsLoading() && BuildTargets.Num() == 0 && !BuildTargetName.IsEmpty())
+			{
+				BuildTargets.Add(BuildTargetName);
+			}
+
+			Archive << AdditionalTargetCommandLineParameters;
+		}
+
 		DefaultLaunchRole->Serialize(Archive);
 
 		// serialize launch roles
@@ -1348,6 +1371,17 @@ public:
 			Writer.WriteArrayEnd();
 		}
 
+		// write the build targets
+		if (BuildTargets.Num() > 0)
+		{
+			Writer.WriteArrayStart("BuildTargets");
+			for (auto Value : BuildTargets)
+			{
+				Writer.WriteValue(Value);
+			}
+			Writer.WriteArrayEnd();
+		}
+
 		Writer.WriteValue("DeployStreamingServer", DeployStreamingServer);
 		Writer.WriteValue("DeployWithUnrealPak", DeployWithUnrealPak);
 		Writer.WriteValue("DeployedDeviceGroupId", DeployedDeviceGroupId.ToString());
@@ -1388,7 +1422,6 @@ public:
 		Writer.WriteValue("UseIoStore", bUseIoStore);
 		Writer.WriteValue("MakeBinaryConfig", bMakeBinaryConfig);
 		Writer.WriteValue("BuildTargetSpecified", BuildTargetSpecified);
-		Writer.WriteValue("BuildTargetName", BuildTargetName);
 		Writer.WriteValue("UseZenStore", bUseZenStore);
 		Writer.WriteValue("ImportingZenSnapshot", bImportingZenSnapshot);
 		Writer.WriteValue("UseZenPakStreaming", bUseZenPakStreaming);
@@ -1400,6 +1433,16 @@ public:
 		Writer.WriteValue("CustomBoolProperties", CustomBoolProperties);
 		Writer.WriteValue("FastIterate", bFastIterate);
 
+		if (AdditionalTargetCommandLineParameters.Num() > 0)
+		{
+			Writer.WriteObjectStart("AdditionalTargetCommandLineParameters");
+			for (const TTuple<EBuildTargetType,FString>& Itr : AdditionalTargetCommandLineParameters)
+			{
+				Writer.WriteValue(LexToString(Itr.Key), Itr.Value);
+			}
+			Writer.WriteObjectEnd();
+		}
+
 		// serialize the default launch role
 		DefaultLaunchRole->Save(Writer, TEXT("DefaultRole"));
 
@@ -2119,10 +2162,18 @@ public:
 			bMakeBinaryConfig = Object.GetBoolField(TEXT("MakeBinaryConfig"));
 		}
 
+		FString BuildTargetName;
 		if (Version >= LAUNCHERSERVICES_ADDBUILDTARGETNAME)
 		{
 			BuildTargetSpecified = Object.GetBoolField(TEXT("BuildTargetSpecified"));
-			BuildTargetName = Object.GetStringField(TEXT("BuildTargetName"));
+			if (Version < LAUNCHERSERVICES_ADDMULTIPLEBUILDTARGETS)
+			{
+				BuildTargetName = Object.GetStringField(TEXT("BuildTargetName"));
+				if (!BuildTargetName.IsEmpty())
+				{
+					BuildTargets.Add(BuildTargetName);
+				}
+			}
 		}
 
 		if (Version >= LAUNCHERSERVICES_ADDZENOPTIONS)
@@ -2196,6 +2247,40 @@ public:
 			bFastIterate = Object.GetBoolField(TEXT("FastIterate"));
 		}
 
+		if (Version >= LAUNCHERSERVICES_ADDMULTIPLEBUILDTARGETS)
+		{
+			BuildTargets.Reset();
+			const TArray<TSharedPtr<FJsonValue>>* JsonValues = NULL;
+			if (Object.TryGetArrayField(TEXT("BuildTargets"), JsonValues))
+			{
+				for (auto Value : *JsonValues)
+				{
+					BuildTargets.Add(Value->AsString());
+				}
+			}
+
+			if (BuildTargets.Num() == 0 && !BuildTargetName.IsEmpty())
+			{
+				BuildTargets.Add(BuildTargetName);
+			}
+
+
+			AdditionalTargetCommandLineParameters.Reset();
+			const TSharedPtr<FJsonObject>* AdditionalTargetCommandLineParametersObject;
+			if (Object.TryGetObjectField(TEXT("AdditionalTargetCommandLineParameters"), AdditionalTargetCommandLineParametersObject))
+			{
+				for (const auto& KeyPairs : (*AdditionalTargetCommandLineParametersObject)->Values)
+				{
+					EBuildTargetType BuildTarget;
+					if (LexTryParseString(BuildTarget, *KeyPairs.Key))
+					{
+						AdditionalTargetCommandLineParameters.Add(BuildTarget, KeyPairs.Value->AsString());
+					}
+				}
+			}
+		}
+
+
 		// load the default launch role
 		TSharedPtr<FJsonObject> Role = Object.GetObjectField(TEXT("DefaultRole"));
 		DefaultLaunchRole->Load(*(Role.Get()));
@@ -2289,6 +2374,7 @@ public:
 		CookedCultures.Add(I18N.GetCurrentCulture()->GetName());
 		CookedMaps.Reset();
 		CookedPlatforms.Reset();
+		BuildTargets.Reset();
 		bSkipCookingEditorContent = false;
         ForceClose = true;
         Timeout = 60;
@@ -2393,6 +2479,28 @@ public:
 		}
 	}
 
+	virtual void SetAdditionalTargetCommandLineParameters(const FString& Params, EBuildTargetType BuildTargetType ) override
+	{
+		check( BuildTargetType != EBuildTargetType::Unknown);
+
+		FString CurrentParams = GetAdditionalTargetCommandLineParameters(BuildTargetType);
+		if (CurrentParams != Params)
+		{
+			if (Params.IsEmpty())
+			{
+				AdditionalTargetCommandLineParameters.Remove(BuildTargetType);
+			}
+			else
+			{
+				AdditionalTargetCommandLineParameters.Add(BuildTargetType, Params);
+			}
+
+			Validate();
+		}
+	}
+
+
+
 	virtual void SetBuildConfiguration( EBuildConfiguration Configuration ) override
 	{
 		if (BuildConfiguration != Configuration)
@@ -2423,9 +2531,16 @@ public:
 
 	virtual void SetBuildTarget( const FString& TargetName ) override
 	{
-		if (BuildTargetName != TargetName)
+		FString BuildTargetName = (BuildTargets.Num() > 0) ? BuildTargets[0] : FString();
+		if (BuildTargetName != TargetName || BuildTargets.Num() > 1)
 		{
-			BuildTargetName = TargetName;
+			UE_CLOG(BuildTargets.Num() > 0, LogTemp, Warning, TEXT("SetBuildTarget removed multiple targets - legacy project launcher?"));
+
+			BuildTargets.Reset();
+			if (!TargetName.IsEmpty())
+			{
+				BuildTargets.Add(TargetName);
+			}
 
 			Validate();
 		}
@@ -3081,6 +3196,37 @@ public:
 		return bWantAutomatedTestBuild;
 	}
 
+	virtual TArray<FString> GetBuildTargets() const
+	{
+		if (BuildTargetSpecified)
+		{
+			return BuildTargets;
+		}
+		else
+		{
+			return TArray<FString>{LauncherProfileManager->GetBuildTarget()};
+		}
+	}
+
+	virtual void AddBuildTarget( const FString& InBuildTarget )
+	{
+		BuildTargets.AddUnique(InBuildTarget);
+		Validate();
+	}
+
+	virtual void RemoveBuildTarget( const FString& InBuildTarget )
+	{
+		BuildTargets.Remove(InBuildTarget);
+		Validate();
+	}
+
+	virtual void ClearBuildTargets()
+	{
+		BuildTargets.Reset();
+		Validate();
+	}
+
+
 	//~ End ILauncherProfile Interface
 
 protected:
@@ -3354,7 +3500,21 @@ protected:
 		bool bBuildTargetIsRequired = false;
 		bool bBuildTargetIsSelected = false;
 
-		FString BuildTarget = GetBuildTarget();
+		TArray<FString> CurrentBuildTargets;
+		if (BuildTargetSpecified)
+		{
+			CurrentBuildTargets = BuildTargets;
+			if (CurrentBuildTargets.IsEmpty())
+			{
+				CurrentBuildTargets.Add(FString());
+			}
+		}
+		else
+		{
+			CurrentBuildTargets.Add(LauncherProfileManager->GetBuildTarget());
+		}
+
+
 		TSet<EBuildTargetType> CookTargetTypes = GetCookTargetTypes();
 		if (CookTargetTypes.Num() == 0)
 		{
@@ -3368,31 +3528,38 @@ protected:
 		{
 			if (CookTargetTypes.Contains(Target.Type))
 			{
-				if (BuildTarget.IsEmpty())
+				for (const FString& BuildTarget : CurrentBuildTargets)
 				{
-					// most target types can have a default build target specified in engine ini so no need to enforce this (UAT will give informative error if the ini isn't set up)
-					bool bSupportsDefaultBuildTarget = (Target.Type != EBuildTargetType::Program); 
-
-					if (!bSupportsDefaultBuildTarget && BuildTargetNames.Contains(Target.Name))
+					if (BuildTarget.IsEmpty())
 					{
-						bBuildTargetIsRequired = true;
-						break;
+						// most target types can have a default build target specified in engine ini so no need to enforce this (UAT will give informative error if the ini isn't set up)
+						bool bSupportsDefaultBuildTarget = (Target.Type != EBuildTargetType::Program); 
+
+						if (!bSupportsDefaultBuildTarget && BuildTargetNames.Contains(Target.Name))
+						{
+							bBuildTargetIsRequired = true;
+							break;
+						}
 					}
-				}
-				else
-				{
-					if (Target.Name == BuildTarget)
+					else
 					{
-						bBuildTargetIsSelected = true;
-						break;
+						if (Target.Name == BuildTarget)
+						{
+							bBuildTargetIsSelected = true;
+							break;
+						}
 					}
 				}
 			}
 		}
 
-		if (!BuildTarget.IsEmpty() && !bBuildTargetIsSelected)
+		for (const FString& BuildTarget : CurrentBuildTargets)
 		{
-			ValidationErrors.Add(ELauncherProfileValidationErrors::BuildTargetCookVariantMismatch);
+			if (!BuildTarget.IsEmpty() && !bBuildTargetIsSelected)
+			{
+				ValidationErrors.Add(ELauncherProfileValidationErrors::BuildTargetCookVariantMismatch);
+				break;
+			}
 		}
 
 		if (bBuildTargetIsRequired)
@@ -3476,7 +3643,7 @@ protected:
 	{
 		RefreshValidBuildTargets();
 
-		BuildTargetName.Empty();
+		BuildTargets.Reset();
 		BuildTargetSpecified = ProjectSpecified; // if this is for 'Any Project' then it should use the fallback build target
 	}
 
@@ -3493,8 +3660,8 @@ private:
 	// Holds a flag indicating whether the build target is specified by this profile.
 	bool BuildTargetSpecified;
 
-	// Holds the name of the target (matching a .cs file) to build. Needed when multiple targets of a type exist
-	FString BuildTargetName;
+	// Holds the names of the build targets (matching a .cs file) to build. Needed when multiple targets of a type exist
+	TArray<FString> BuildTargets;
 
 	// Holds the build mode.
 	// Holds the build configuration name of the cooker.
@@ -3693,6 +3860,9 @@ private:
 	// Additional command line parameters to set for the application when it launches
 	FString AdditionalCommandLineParameters;
 
+	// Additional command line parameters to set for a given build target type when it launches
+	TMap<EBuildTargetType,FString> AdditionalTargetCommandLineParameters;
+
 	// Use I/O store.
 	bool bUseIoStore;
 	
diff --git a/Engine/Source/Developer/LauncherServices/Public/ILauncherProfile.h b/Engine/Source/Developer/LauncherServices/Public/ILauncherProfile.h
--- a/Engine/Source/Developer/LauncherServices/Public/ILauncherProfile.h
+++ b/Engine/Source/Developer/LauncherServices/Public/ILauncherProfile.h
@@ -650,6 +650,7 @@ public:
 	 * @Return Target name to build/run (it would match a .Target.cs file)
 	 * @see SetBuildTarget
 	 */
+	UE_DEPRECATED(5.8, "Use GetBuildTargets instead")
 	virtual FString GetBuildTarget() const = 0;
 
 	/**
@@ -839,6 +840,15 @@ public:
 	 */
 	virtual FString GetAdditionalCommandLineParameters() const = 0;
 
+	/**
+	 * Gets the additional command line parameters that will be used for the given target type when the app launches.
+	 *
+ 	 * @param	BuildTargetType	The type of build target to query
+	 * @return The additional command line parameters
+	 * @see SetAdditionalTargetCommandLineParameters
+	 */
+	virtual FString GetAdditionalTargetCommandLineParameters( EBuildTargetType BuildTargetType ) const = 0;
+
     /**
      * Gets the timeout time for the cook on the fly server.
      *
@@ -1182,6 +1192,7 @@ public:
 	 * @param TargetName The target name to set (it would match a .Target.cs file)
 	 * @see GetBuildTarget
 	 */
+	UE_DEPRECATED(5.8, "Use AddBuildTarget instead")
 	virtual void SetBuildTarget( const FString& TargetName ) = 0;
 
 	/**
@@ -1442,6 +1453,16 @@ public:
 	 */
 	virtual void SetAdditionalCommandLineParameters(const FString& Params) = 0;
 
+	/**
+	 * Sets the additional command line parameters that will be used for the given target type when the app launches.
+	 *
+	 * @param	Params	The additional command line parameters to use
+	 * @param	BuildTargetType	The type of build target to modify
+	 * @see GetAdditionalTargetCommandLineParameters
+	 */
+	virtual void SetAdditionalTargetCommandLineParameters(const FString& Params, EBuildTargetType BuildTargetType ) = 0;
+
+
 	/**
 	 * Sets whether to use a streaming file server.
 	 *
@@ -1675,6 +1696,35 @@ public:
 	 */
 	virtual bool IsUsingAutomatedTestBuild() const = 0;
 
+
+
+
+	/**
+	 * Gets the names of the build targets
+	 * @see AddBuildTarget, RemoveBuildTarget, ClearBuildTargets
+	 */
+	virtual TArray<FString> GetBuildTargets() const = 0;
+
+	/**
+	 * Add a build target to the list
+	 * @see GetBuildTargets, RemoveBuildTarget, ClearBuildTargets
+	 */
+	virtual void AddBuildTarget( const FString& InBuildTarget ) = 0;
+
+	/**
+	 * Remove a build target from the list
+	 * @see GetBuildTargets, AddBuildTarget, ClearBuildTargets
+	 */
+	virtual void RemoveBuildTarget( const FString& InBuildTarget ) = 0;
+
+	/**
+	 * Clears the list of build targets
+	 * @see GetBuildTargets, AddBuildTarget, RemoveBuildTarget
+	 */
+	virtual void ClearBuildTargets() = 0;
+
+
+
 public:
 	/**
 	 * Helper function to get all of the build targets available for this profile, based on its current project & cook platforms
diff --git a/Engine/Source/Developer/LauncherServices/Public/ILauncherProfileManager.h b/Engine/Source/Developer/LauncherServices/Public/ILauncherProfileManager.h
--- a/Engine/Source/Developer/LauncherServices/Public/ILauncherProfileManager.h
+++ b/Engine/Source/Developer/LauncherServices/Public/ILauncherProfileManager.h
@@ -53,9 +53,9 @@ DECLARE_MULTICAST_DELEGATE(FOnLauncherProfileManagerProjectChanged);
 /**
  * Declares a delegate to be invoked to post-process a command line before launching
  *
- * The first parameter is the profile that is being launched, the second parameter is the command line to modify.
+ * The first parameter is the profile that is being launched, the second parameter is the build target type (or 'unknown' for global arguments), and the third parameter is the command line to modify.
  */
-DECLARE_MULTICAST_DELEGATE_TwoParams(FOnLauncherProfilePostProcessCommandLine, const ILauncherProfileRef&, FString& );
+DECLARE_MULTICAST_DELEGATE_ThreeParams(FOnLauncherProfilePostProcessCommandLine, const ILauncherProfileRef&, EBuildTargetType, FString& );
 
 /**
  * Declares a delegate to be invoked to post-process a command line before launching and automated test
diff --git a/Engine/Source/Developer/LegacyProjectLauncher/Private/Widgets/Shared/SProjectLauncherBuildTargetSelector.cpp b/Engine/Source/Developer/LegacyProjectLauncher/Private/Widgets/Shared/SProjectLauncherBuildTargetSelector.cpp
--- a/Engine/Source/Developer/LegacyProjectLauncher/Private/Widgets/Shared/SProjectLauncherBuildTargetSelector.cpp
+++ b/Engine/Source/Developer/LegacyProjectLauncher/Private/Widgets/Shared/SProjectLauncherBuildTargetSelector.cpp
@@ -114,7 +114,9 @@ FText SProjectLauncherBuildTargetSelector::HandleBuildTargetComboButtonText() co
 		ILauncherProfilePtr Profile = GetProfile();
 		if (Profile.IsValid() && Profile->HasBuildTargetSpecified())
 		{
-			BuildTarget = Profile->GetBuildTarget();
+			const TArray<FString> BuildTargets = Profile->GetBuildTargets();
+
+			BuildTarget = BuildTargets.Num() > 0 ? BuildTargets[0] : FString();
 		}
 	}
 	else
@@ -170,7 +172,11 @@ void SProjectLauncherBuildTargetSelector::SetBuildTarget(const FString& BuildTar
 		ILauncherProfilePtr Profile = GetProfile();
 		if (Profile.IsValid())
 		{
-			Profile->SetBuildTarget(BuildTarget);
+			Profile->ClearBuildTargets();
+			if (!BuildTarget.IsEmpty())
+			{
+				Profile->AddBuildTarget(BuildTarget);
+			}
 			Profile->SetBuildTargetSpecified(true);
 		}
 	}
diff --git a/Engine/Source/Editor/UnrealEd/Private/PlayLevelViaLauncher.cpp b/Engine/Source/Editor/UnrealEd/Private/PlayLevelViaLauncher.cpp
--- a/Engine/Source/Editor/UnrealEd/Private/PlayLevelViaLauncher.cpp
+++ b/Engine/Source/Editor/UnrealEd/Private/PlayLevelViaLauncher.cpp
@@ -292,7 +292,8 @@ void UEditorEngine::StartPlayUsingLauncherSession(FRequestPlaySessionParams& InR
 		const FTargetInfo* TargetInfo = GetDefault<UPlatformsMenuSettings>()->GetLaunchOnTargetInfo();
 		if (TargetInfo != nullptr)
 		{
-			LauncherProfile->SetBuildTarget(TargetInfo->Name);
+			LauncherProfile->ClearBuildTargets();
+			LauncherProfile->AddBuildTarget(TargetInfo->Name);
 			LauncherProfile->SetBuildTargetSpecified(true);
 		}
 	}
