Submission Package content scheme for Project Launcher

#rnx
#Synced-CL 48448085

diff --git a/Engine/Plugins/Developer/ProjectLauncher/Source/ProjectLauncher/Private/Model/ProjectLauncherModel.cpp b/Engine/Plugins/Developer/ProjectLauncher/Source/ProjectLauncher/Private/Model/ProjectLauncherModel.cpp
--- a/Engine/Plugins/Developer/ProjectLauncher/Source/ProjectLauncher/Private/Model/ProjectLauncherModel.cpp
+++ b/Engine/Plugins/Developer/ProjectLauncher/Source/ProjectLauncher/Private/Model/ProjectLauncherModel.cpp
@@ -129,9 +129,6 @@ namespace ProjectLauncher
 	bool FModel::IsAdvancedProfile(const ILauncherProfileRef& Profile) const
 	{
 		return 
-				Profile->IsCreatingDLC() ||
-				Profile->IsCreatingReleaseVersion() ||
-				Profile->IsGeneratingPatch() ||				
 				Profile->GetPackagingMode() == ELauncherProfilePackagingModes::SharedRepository ||
 				Profile->GetDeploymentMode() == ELauncherProfileDeploymentModes::CopyRepository ||
 				Profile->GetLaunchMode() == ELauncherProfileLaunchModes::CustomRoles
@@ -672,7 +669,14 @@ PRAGMA_ENABLE_DEPRECATION_WARNINGS
 		}
 		else if (Profile->GetPackagingMode() != ELauncherProfilePackagingModes::DoNotPackage)
 		{
-			return EContentScheme::DevelopmentPackage;
+			if (Profile->IsCreatingReleaseVersion() || Profile->IsGeneratingPatch() || Profile->IsCreatingDLC())
+			{
+				return EContentScheme::SubmissionPackage;
+			}
+			else
+			{
+				return EContentScheme::DevelopmentPackage;
+			}
 		}
 		else if (Profile->IsPackingWithUnrealPak())
 		{
@@ -696,16 +700,20 @@ PRAGMA_ENABLE_DEPRECATION_WARNINGS
 
 	void FModel::SetProfileContentScheme(EContentScheme ContentScheme, const ILauncherProfileRef& Profile, bool bWantToCook, ELauncherProfileDeploymentModes::Type DefaultDeploymentMode)
 	{
-		bool bPakFiles = (ContentScheme == EContentScheme::PakFiles || ContentScheme == EContentScheme::DevelopmentPackage);
+		bool bPakFiles = (ContentScheme == EContentScheme::PakFiles || ContentScheme == EContentScheme::DevelopmentPackage || ContentScheme == EContentScheme::SubmissionPackage);
 		bool bUseZen = (ContentScheme != EContentScheme::LooseFiles && ContentScheme != EContentScheme::CookOnTheFly);
 		bool bCOTF = (ContentScheme == EContentScheme::CookOnTheFly);
-		bool bPackage = (ContentScheme == EContentScheme::DevelopmentPackage);
+		bool bPackage = (ContentScheme == EContentScheme::DevelopmentPackage || ContentScheme == EContentScheme::SubmissionPackage);
 		bool bZenPakStreaming = (ContentScheme == EContentScheme::ZenPakStreaming);
+		bool bSubmissionPackage = (ContentScheme == EContentScheme::SubmissionPackage);
 		// FIXME: may need to turn off Zen/Pak in UAT somehow if it's set in project defaults?
 
 		Profile->SetUseZenPakStreaming(bZenPakStreaming);
 		Profile->SetDeployWithUnrealPak(bPakFiles);
 		Profile->SetUseZenStore(bUseZen);
+		Profile->SetCreateReleaseVersion(bSubmissionPackage);
+		Profile->SetGeneratePatch(false);
+		Profile->SetCreateDLC(false);
 		if (!bPakFiles)
 		{
 			Profile->SetGenerateChunks(false);
@@ -1054,6 +1062,7 @@ PRAGMA_ENABLE_DEPRECATION_WARNINGS
 			case EContentScheme::ZenStreaming:       return LOCTEXT("ContentSchemeZenStreaming","Zen Streaming");
 			case EContentScheme::ZenPakStreaming:    return LOCTEXT("ContentSchemeZenPakStreaming","Zen Pak Streaming"); 
 			case EContentScheme::DevelopmentPackage: return LOCTEXT("ContentSchemeDevPackage","Development Package");
+			case EContentScheme::SubmissionPackage:  return LOCTEXT("ContentSchemeSubPackage","Submission Package");
 			case EContentScheme::LooseFiles:         return LOCTEXT("ContentSchemeLooseFiles","Loose Files (legacy)");
 			case EContentScheme::CookOnTheFly:       return LOCTEXT("ContentSchemeCOTF","Cook On The Fly"); 
 		}
@@ -1071,6 +1080,7 @@ PRAGMA_ENABLE_DEPRECATION_WARNINGS
 			case EContentScheme::ZenStreaming:       return LOCTEXT("ContentSchemeTipZenStreaming","Stream cooked game content from Zen Server");
 			case EContentScheme::ZenPakStreaming:    return LOCTEXT("ContentSchemeTipZenPakStreaming","Stream an existing Pak Files build via a Zen");
 			case EContentScheme::DevelopmentPackage: return LOCTEXT("ContentSchemeTipDevPackage","Package cooked game content into a single installable package file for development purposes, where available");
+			case EContentScheme::SubmissionPackage:  return LOCTEXT("ContentSchemeTipSubPackage","Package cooked game content into a single package or patch for submission to a game store, where available");
 			case EContentScheme::LooseFiles:         return LOCTEXT("ContentSchemeTipLooseFiles","Store cooked game assets in individual files (legacy - recommend moving to Zen Streaming. This option will not work if the project is already configured with 'Use Zen Store')");
 			case EContentScheme::CookOnTheFly:       return LOCTEXT("ContentSchemeTipCOTF","Only cook game assets when the game requires them, and send them over the network (legacy - slow)");
 		}
@@ -1091,6 +1101,7 @@ PRAGMA_ENABLE_DEPRECATION_WARNINGS
 			case ProjectLauncher::EContentScheme::ZenStreaming:       return TEXT("ZenStreaming");
 			case ProjectLauncher::EContentScheme::ZenPakStreaming:    return TEXT("ZenPakStreaming");
 			case ProjectLauncher::EContentScheme::DevelopmentPackage: return TEXT("DevelopmentPackage");
+			case ProjectLauncher::EContentScheme::SubmissionPackage:  return TEXT("SubmissionPackage");
 			case ProjectLauncher::EContentScheme::LooseFiles:         return TEXT("LooseFiles");
 			case ProjectLauncher::EContentScheme::CookOnTheFly:       return TEXT("CookOnTheFly");
 			default: return TEXT("Unknown");
@@ -1121,6 +1132,11 @@ PRAGMA_ENABLE_DEPRECATION_WARNINGS
 			OutContentScheme = ProjectLauncher::EContentScheme::DevelopmentPackage;
 			return true;
 		}
+		else if (FCString::Stricmp(String, TEXT("SubmissionPackage")) == 0)
+		{
+			OutContentScheme = ProjectLauncher::EContentScheme::SubmissionPackage;
+			return true;
+		}
 		else if (FCString::Stricmp(String, TEXT("LooseFiles")) == 0) 
 		{
 			OutContentScheme = ProjectLauncher::EContentScheme::LooseFiles;
diff --git a/Engine/Plugins/Developer/ProjectLauncher/Source/ProjectLauncher/Private/ProfileTree/CustomProfileTreeBuilder.cpp b/Engine/Plugins/Developer/ProjectLauncher/Source/ProjectLauncher/Private/ProfileTree/CustomProfileTreeBuilder.cpp
--- a/Engine/Plugins/Developer/ProjectLauncher/Source/ProjectLauncher/Private/ProfileTree/CustomProfileTreeBuilder.cpp
+++ b/Engine/Plugins/Developer/ProjectLauncher/Source/ProjectLauncher/Private/ProfileTree/CustomProfileTreeBuilder.cpp
@@ -31,6 +31,9 @@ namespace ProjectLauncher
 		AddImportZenSnapshotProperty(ContentSchemeHeader);
 		AddZenPakStreamingPathProperty(ContentSchemeHeader);
 
+		FLaunchProfileTreeNode& SubmissionHeader = TreeData->AddHeading( TEXT("Release"), LOCTEXT("ReleaseHeading", "Release / DLC / Patching") );
+		AddSubmissionPackageProperty(SubmissionHeader);
+
 		FLaunchProfileTreeNode& CookingHeader = TreeData->AddHeading( TEXT("Cooking"), LOCTEXT("CookingHeading", "Cook") );
 		AddCookProperty(CookingHeader);
 		AddIncrementalCookProperty(CookingHeader);
diff --git a/Engine/Plugins/Developer/ProjectLauncher/Source/ProjectLauncher/Private/ProfileTree/GenericProfileTreeBuilder.cpp b/Engine/Plugins/Developer/ProjectLauncher/Source/ProjectLauncher/Private/ProfileTree/GenericProfileTreeBuilder.cpp
--- a/Engine/Plugins/Developer/ProjectLauncher/Source/ProjectLauncher/Private/ProfileTree/GenericProfileTreeBuilder.cpp
+++ b/Engine/Plugins/Developer/ProjectLauncher/Source/ProjectLauncher/Private/ProfileTree/GenericProfileTreeBuilder.cpp
@@ -38,7 +38,7 @@ namespace ProjectLauncher
 
 		ForPak = [this]()
 		{
-			return (ContentScheme == EContentScheme::PakFiles || ContentScheme == EContentScheme::DevelopmentPackage);
+			return (ContentScheme == EContentScheme::PakFiles || ContentScheme == EContentScheme::DevelopmentPackage || ContentScheme == EContentScheme::SubmissionPackage);
 		};
 
 		ForZen = [this]()
@@ -81,6 +81,11 @@ namespace ProjectLauncher
 			return GetIsRunning();
 		};
 
+		ForSubmissionPackage = [this]()
+		{
+			return (ContentScheme == EContentScheme::SubmissionPackage);
+		};
+
 		EmptyString = []()
 		{
 			return FString();
@@ -340,6 +345,183 @@ namespace ProjectLauncher
 		);
 	}
 
+	void FGenericProfileTreeBuilder::AddSubmissionPackageProperty( FLaunchProfileTreeNode& HeadingNode )
+	{
+		auto TryGetTargetPlatform = [this](const ITargetPlatform*& OutPlatform)
+		{
+			const TArray<FString>& CookedPlatforms = Profile->GetCookedPlatforms();
+			if (CookedPlatforms.Num() == 1)
+			{
+				const TArray<ITargetPlatform*>& TargetPlatforms = GetTargetPlatformManager()->GetTargetPlatforms();
+				for (const ITargetPlatform* TargetPlatform : TargetPlatforms)
+				{
+					if (CookedPlatforms[0] == TargetPlatform->PlatformName())
+					{
+						OutPlatform = TargetPlatform;
+						return true;
+					}
+				}
+			}
+
+			OutPlatform = nullptr;
+			return false;
+		};
+
+		auto HasNonTrivialPatchGeneration = [TryGetTargetPlatform]()
+		{
+			const ITargetPlatform* Platform;
+			return TryGetTargetPlatform(Platform) && Platform->HasNonTrivialPatchGeneration();
+		};
+
+		auto IsBasedOnReleaseVersionVisible = [this]
+		{
+			return ForSubmissionPackage() && (Profile->IsCreatingDLC() || Profile->IsGeneratingPatch());
+		};
+
+		auto IsOriginalReleaseVersionVisible = [this, TryGetTargetPlatform]
+		{
+			const ITargetPlatform* Platform;
+			return ForSubmissionPackage() && TryGetTargetPlatform(Platform) && Platform->RequiresOriginalReleaseVersionForPatch() && Profile->IsGeneratingPatch();
+		};
+
+
+		// ... warning banner ...
+		HeadingNode.AddWidget( FText::GetEmpty(),
+			{
+				.IsVisible = [this,HasNonTrivialPatchGeneration]() { return ForSubmissionPackage() && HasNonTrivialPatchGeneration(); }
+			},
+			SNew(SBorder)
+			.BorderImage(FAppStyle::GetBrush("ToolPanel.LightGroupBorder"))
+			.BorderBackgroundColor(FAppStyle::Get().GetSlateColor("Colors.Warning"))
+			[
+				SNew(STextBlock)
+				.Margin(FMargin(8,2))
+				.ColorAndOpacity(FLinearColor::White)
+				.AutoWrapText(true)
+				.Text(LOCTEXT("ComplexPkgInfo", "This platform has some additional complexities for creating submission packages and patches. Please refer to the relevant documentation for details."))
+			]
+		)
+
+		// ... release version ...
+		.AddBoolean( LOCTEXT("CreateReleaseLabel", "Create Release Version"),
+			{
+				.GetValue = [this]()				{ return Profile->IsCreatingReleaseVersion(); },
+				.SetValue = [this](bool bValue)		{ Profile->SetCreateReleaseVersion(bValue); },
+				.GetDefaultValue = [this]()			{ return DefaultProfile->IsCreatingReleaseVersion(); },
+				.IsVisible = ForSubmissionPackage,
+			}
+		)
+		.AddString( LOCTEXT("ReleaseVersionNameLabel", "Release Version Name"),
+			{
+				.GetValue = [this]()				{ return Profile->GetCreateReleaseVersionName(); },
+				.SetValue = [this](FString Value)	{ Profile->SetCreateReleaseVersionName(Value); },
+				.GetDefaultValue = [this]()			{ return DefaultProfile->GetCreateReleaseVersionName(); },
+				.IsVisible = ForSubmissionPackage,
+				.IsEnabled = [this]()				{ return Profile->IsCreatingReleaseVersion() || Profile->IsGeneratingPatch(); },
+			}
+		)
+
+		// ... various release/distribution options ....
+		.AddBoolean( LOCTEXT("EncryptIniLabel", "Encrypt ini files"),
+			{
+				.GetValue = [this]()				{ return Profile->IsEncryptingIniFiles(); },
+				.SetValue = [this](bool bValue)		{ Profile->SetEncryptingIniFiles(bValue); },
+				.GetDefaultValue = [this]()			{ return DefaultProfile->IsEncryptingIniFiles(); },
+				.IsVisible = ForSubmissionPackage,
+			}
+		)
+		.AddBoolean( LOCTEXT("ForDistributionLabel", "This build is for distribution to the public"),
+			{
+				.GetValue = [this]()				{ return Profile->IsForDistribution(); },
+				.SetValue = [this](bool bValue)		{ Profile->SetForDistribution(bValue); },
+				.GetDefaultValue = [this]()			{ return DefaultProfile->IsForDistribution(); },
+				.IsVisible = ForSubmissionPackage,
+			}
+		)
+		.AddBoolean( LOCTEXT("IncludePrereqLabel", "Include an installer for prerequisites"),
+			{
+				.GetValue = [this]()				{ return Profile->IsIncludingPrerequisites(); },
+				.SetValue = [this](bool bValue)		{ Profile->SetIncludePrerequisites(bValue); },
+				.GetDefaultValue = [this]()			{ return DefaultProfile->IsIncludingPrerequisites(); },
+				.IsVisible = ForSubmissionPackage,
+			}
+		)
+		.AddBoolean( LOCTEXT("BinaryConfigLabel", "Make a binary config file for faster runtime settings"),
+			{
+				.GetValue = [this]()				{ return Profile->MakeBinaryConfig(); },
+				.SetValue = [this](bool bValue)		{ Profile->SetMakeBinaryConfig(bValue); },
+				.GetDefaultValue = [this]()			{ return DefaultProfile->MakeBinaryConfig(); },
+				.IsVisible = ForSubmissionPackage,
+			}
+		)
+
+		// ... patch/DLC options ...
+		.AddBoolean( LOCTEXT("CreatePatchLabel", "Create Patch"),
+			{
+				.GetValue = [this]()				{ return Profile->IsGeneratingPatch(); },
+				.SetValue = [this](bool bValue)		{ Profile->SetGeneratePatch(bValue); },
+				.GetDefaultValue = [this]()			{ return DefaultProfile->IsGeneratingPatch(); },
+				.IsVisible = [this]()				{ return ForSubmissionPackage(); },
+				.IsEnabled = [this]()				{ return !Profile->IsCreatingDLC(); },
+			}
+		)
+		.AddBoolean( LOCTEXT("CreateDLCLabel", "Create DLC"),
+			{
+				.GetValue = [this]()				{ return Profile->IsCreatingDLC(); },
+				.SetValue = [this](bool bValue)		{ Profile->SetCreateDLC(bValue); },
+				.GetDefaultValue = [this]()			{ return DefaultProfile->IsCreatingDLC(); },
+				.IsVisible = [this]()				{ return ForSubmissionPackage(); },
+				.IsEnabled = [this]()				{ return !Profile->IsGeneratingPatch(); },
+			}
+		)
+
+		.AddString( LOCTEXT("BasedOnReleaseVersionNameLabel", "Based On Release Version Name"),
+			{
+				.GetValue = [this]()				{ return Profile->GetBasedOnReleaseVersionName(); },
+				.SetValue = [this](FString Value)	{ Profile->SetBasedOnReleaseVersionName(Value); },
+				.GetDefaultValue = [this]()			{ return DefaultProfile->GetBasedOnReleaseVersionName(); },
+				.IsVisible = IsBasedOnReleaseVersionVisible,
+			}
+		)
+		.AddString( LOCTEXT("OriginalReleaseVersionNameLabel", "Original Release Version Name"),
+			{
+				.GetValue = [this]()				{ return Profile->GetOriginalReleaseVersionName(); },
+				.SetValue = [this](FString Value)	{ Profile->SetOriginalReleaseVersionName(Value); },
+				.GetDefaultValue = [this]()			{ return DefaultProfile->GetOriginalReleaseVersionName(); },
+				.IsVisible = IsOriginalReleaseVersionVisible,
+			}
+		)
+
+		.AddBoolean( LOCTEXT("AddPatchLevelLabel", "Add Patch Level (not generally needed)"),
+			{
+				.GetValue = [this]()				{ return Profile->ShouldAddPatchLevel(); },
+				.SetValue = [this](bool bValue)		{ Profile->SetAddPatchLevel(bValue); },
+				.GetDefaultValue = [this]()			{ return DefaultProfile->ShouldAddPatchLevel(); },
+				.IsVisible = [this]()				{ return ForSubmissionPackage() && Profile->IsGeneratingPatch(); },
+			}
+		)
+
+		.AddString( LOCTEXT("DLCNameLabel", "DLC Name"),
+			{
+				.GetValue = [this]()				{ return Profile->GetDLCName(); },
+				.SetValue = [this](FString Value)	{ Profile->SetDLCName(Value); },
+				.GetDefaultValue = [this]()			{ return DefaultProfile->GetDLCName(); },
+				.IsVisible = [this]()				{ return ForSubmissionPackage() && Profile->IsCreatingDLC(); },
+			}
+		)
+		.AddBoolean( LOCTEXT("DLCIncludeEngineLabel", "DLC Includes Engine Content"),
+			{
+				.GetValue = [this]()				{ return Profile->IsDLCIncludingEngineContent(); },
+				.SetValue = [this](bool bValue)		{ Profile->SetDLCIncludeEngineContent(bValue); },
+				.GetDefaultValue = [this]()			{ return DefaultProfile->IsDLCIncludingEngineContent(); },
+				.IsVisible = [this]()				{ return ForSubmissionPackage() && Profile->IsCreatingDLC(); },
+			}
+		)
+
+		;
+	}
+
+
 	void FGenericProfileTreeBuilder::AddIncrementalCookProperty( FLaunchProfileTreeNode& HeadingNode )
 	{
 		HeadingNode.AddBoolean( LOCTEXT("IncrementalCookLabel", "Incremental Cook"),
@@ -490,7 +672,7 @@ namespace ProjectLauncher
 				.GetValue = [this]()			{ return GetDeployToDevice(); },
 				.SetValue = [this](bool bValue)	{ SetDeployToDevice(bValue); },
 				.GetDefaultValue = [this]()		{ return GetDeployToDevice(DefaultProfile); },
-				.IsVisible = [this]()			{ return (ContentScheme != EContentScheme::CookOnTheFly); },
+				.IsVisible = [this]()			{ return (ContentScheme != EContentScheme::CookOnTheFly) && (ContentScheme != EContentScheme::SubmissionPackage); },
 				.Validation = FValidation({ELauncherProfileValidationErrors::CopyToDeviceRequiresCookByTheBook, ELauncherProfileValidationErrors::CopyToDeviceRequiresNoPackaging, ELauncherProfileValidationErrors::ZenPakStreamingRequiresDeployAndLaunch, ELauncherProfileValidationErrors::DeployedDeviceGroupRequired }),
 			}
 		);
@@ -503,6 +685,7 @@ namespace ProjectLauncher
 				.GetValue = [this]()			{ return Profile->IsDeployingIncrementally(); },
 				.SetValue = [this](bool bValue)	{ Profile->SetIncrementalDeploying(bValue); },
 				.GetDefaultValue = [this]()		{ return DefaultProfile->IsDeployingIncrementally(); },
+				.IsVisible = [this]()			{ return (ContentScheme != EContentScheme::CookOnTheFly) && (ContentScheme != EContentScheme::SubmissionPackage); },
 				.IsEnabled = ForDeployment,
 			}
 		);
@@ -512,6 +695,7 @@ namespace ProjectLauncher
 	{
 		HeadingNode.AddWidget( LOCTEXT("TargetDeviceLabel", "Target Device"), 
 			{
+				.IsVisible = [this]()			{ return (ContentScheme != EContentScheme::CookOnTheFly) && (ContentScheme != EContentScheme::SubmissionPackage); },
 				.IsEnabled = ForDeployment,
 				.Validation = FValidation( {ELauncherProfileValidationErrors::BuildTargetCookVariantMismatch, ELauncherProfileValidationErrors::LaunchDeviceIsUnauthorized} ),
 			},
@@ -527,7 +711,7 @@ namespace ProjectLauncher
 				.GetValue = [this]()			{ return GetIsRunning(); },
 				.SetValue = [this](bool bValue)	{ SetIsRunning(bValue); },
 				.GetDefaultValue = [this]()		{ return GetIsRunning(DefaultProfile); },
-				.IsVisible = [this]()           { return Profile->GetAutomatedTests().Num() == 0; }, // run is implied when there are automated tests
+				.IsVisible = [this]()           { return (ContentScheme != EContentScheme::SubmissionPackage) && Profile->GetAutomatedTests().Num() == 0; }, // run is implied when there are automated tests
 				.Validation = FValidation({ELauncherProfileValidationErrors::ZenPakStreamingRequiresDeployAndLaunch, ELauncherProfileValidationErrors::NoLaunchRoleDeviceAssigned}),
 			}
 		);
@@ -541,6 +725,7 @@ namespace ProjectLauncher
 			{
 				.IsDefault = [this]()		{ return Profile->GetDefaultLaunchRole()->GetInitialMap().IsEmpty(); },
 				.SetToDefault = [this]()	{ Profile->GetDefaultLaunchRole()->SetInitialMap(FString()); InitalMapCombo->SetSelectedItem(GetInitialMap()); },
+				.IsVisible = [this]()		{ return (ContentScheme != EContentScheme::SubmissionPackage); },
 				.IsEnabled = ForRun,
 				.Validation = FValidation({ELauncherProfileValidationErrors::InitialMapNotAvailable}),
 			},
@@ -566,6 +751,7 @@ namespace ProjectLauncher
 				.GetValue = [this]()				{ return GetCommandLine(EBuildTargetType::Unknown); },
 				.SetValue = [this](FString Value)	{ SetCommandLine(Value, EBuildTargetType::Unknown); },
 				.GetDefaultValue = EmptyString,
+				.IsVisible = [this]()				{ return (ContentScheme != EContentScheme::SubmissionPackage); },
 				.IsEnabled = ForRun,
 				.Validation = FValidation({ELauncherProfileValidationErrors::MalformedLaunchCommandLine}),
 			}
@@ -936,6 +1122,13 @@ namespace ProjectLauncher
 			return false;
 		}
 
+		// cannot create a submission package from the Basic Launch
+		if (ProfileType == EProfileType::Basic && InContentScheme == EContentScheme::SubmissionPackage)
+		{
+			// don't set a reason - just hide the item for Basic Launch
+			return false;
+		}
+
 		// loose files can't be selected if the project is using zen store because there's no way to opt-out of Zen Store from the UAT command line.
 		if (InContentScheme == EContentScheme::LooseFiles && ProjectSettings.bUseZenStore)
 		{
@@ -959,6 +1152,13 @@ namespace ProjectLauncher
 			return false;
 		}
 
+		// cannot make a submission package outside of the shipping configuration
+		if (InContentScheme == EContentScheme::SubmissionPackage && Profile->GetBuildConfiguration() != EBuildConfiguration::Shipping)
+		{
+			OutReason = LOCTEXT("SubPkgShippingOnly", "Submission packages require a Shipping build");
+			return false;
+		}
+
 		return true;
 	}
 
diff --git a/Engine/Plugins/Developer/ProjectLauncher/Source/ProjectLauncher/Public/Model/ProjectLauncherModel.h b/Engine/Plugins/Developer/ProjectLauncher/Source/ProjectLauncher/Public/Model/ProjectLauncherModel.h
--- a/Engine/Plugins/Developer/ProjectLauncher/Source/ProjectLauncher/Public/Model/ProjectLauncherModel.h
+++ b/Engine/Plugins/Developer/ProjectLauncher/Source/ProjectLauncher/Public/Model/ProjectLauncherModel.h
@@ -34,6 +34,7 @@ namespace ProjectLauncher
 		ZenStreaming,
 		ZenPakStreaming,
 		DevelopmentPackage,
+		SubmissionPackage,
 		LooseFiles,
 		CookOnTheFly,
 		MAX
diff --git a/Engine/Plugins/Developer/ProjectLauncher/Source/ProjectLauncher/Public/ProfileTree/GenericProfileTreeBuilder.h b/Engine/Plugins/Developer/ProjectLauncher/Source/ProjectLauncher/Public/ProfileTree/GenericProfileTreeBuilder.h
--- a/Engine/Plugins/Developer/ProjectLauncher/Source/ProjectLauncher/Public/ProfileTree/GenericProfileTreeBuilder.h
+++ b/Engine/Plugins/Developer/ProjectLauncher/Source/ProjectLauncher/Public/ProfileTree/GenericProfileTreeBuilder.h
@@ -58,6 +58,7 @@ namespace ProjectLauncher
 		UE_API void AddZenSnapshotProperty(FLaunchProfileTreeNode& HeadingNode);
 		UE_API void AddImportZenSnapshotProperty( FLaunchProfileTreeNode& HeadingNode );
 		UE_API void AddZenPakStreamingPathProperty( FLaunchProfileTreeNode& HeadingNode );
+		UE_API void AddSubmissionPackageProperty( FLaunchProfileTreeNode& HeadingNode );
 		UE_API void AddIncrementalCookProperty( FLaunchProfileTreeNode& HeadingNode );
 		UE_API void AddCookProperty( FLaunchProfileTreeNode& HeadingNode );
 		UE_API void AddBuildProperty( FLaunchProfileTreeNode& HeadingNode );
@@ -84,6 +85,7 @@ namespace ProjectLauncher
 		FLaunchProfileTreeNode::FGetBool ForCode;
 		FLaunchProfileTreeNode::FGetBool ForDeployment;
 		FLaunchProfileTreeNode::FGetBool ForRun;
+		FLaunchProfileTreeNode::FGetBool ForSubmissionPackage;
 		FLaunchProfileTreeNode::FGetString EmptyString;
 
 
diff --git a/Engine/Source/Developer/TargetPlatform/Public/Common/TargetPlatformBase.h b/Engine/Source/Developer/TargetPlatform/Public/Common/TargetPlatformBase.h
--- a/Engine/Source/Developer/TargetPlatform/Public/Common/TargetPlatformBase.h
+++ b/Engine/Source/Developer/TargetPlatform/Public/Common/TargetPlatformBase.h
@@ -424,6 +424,11 @@ public:
 		return TPlatformProperties::RequiresOriginalReleaseVersionForPatch();
 	}
 
+	virtual bool HasNonTrivialPatchGeneration() const override
+	{
+		return TPlatformProperties::HasNonTrivialPatchGeneration();
+	}
+
 	virtual bool HasSecurePackageFormat() const override
 	{
 		return TPlatformProperties::HasSecurePackageFormat();
@@ -764,6 +769,10 @@ public:
 	{
 		return TargetPlatformControls->RequiresOriginalReleaseVersionForPatch();
 	}
+	virtual bool HasNonTrivialPatchGeneration() const override 
+	{
+		return TargetPlatformControls->HasNonTrivialPatchGeneration();
+	}
 	virtual bool HasSecurePackageFormat() const override 
 	{
 		return TargetPlatformControls->HasSecurePackageFormat();
diff --git a/Engine/Source/Developer/TargetPlatform/Public/Common/TargetPlatformControlsBase.h b/Engine/Source/Developer/TargetPlatform/Public/Common/TargetPlatformControlsBase.h
--- a/Engine/Source/Developer/TargetPlatform/Public/Common/TargetPlatformControlsBase.h
+++ b/Engine/Source/Developer/TargetPlatform/Public/Common/TargetPlatformControlsBase.h
@@ -319,6 +319,11 @@ public:
 		return TPlatformProperties::RequiresOriginalReleaseVersionForPatch();
 	}
 
+	virtual bool HasNonTrivialPatchGeneration() const override
+	{
+		return TPlatformProperties::HasNonTrivialPatchGeneration();
+	}
+
 	virtual bool HasSecurePackageFormat() const override
 	{
 		return TPlatformProperties::HasSecurePackageFormat();
diff --git a/Engine/Source/Developer/TargetPlatform/Public/Interfaces/ITargetPlatformControls.h b/Engine/Source/Developer/TargetPlatform/Public/Interfaces/ITargetPlatformControls.h
--- a/Engine/Source/Developer/TargetPlatform/Public/Interfaces/ITargetPlatformControls.h
+++ b/Engine/Source/Developer/TargetPlatform/Public/Interfaces/ITargetPlatformControls.h
@@ -288,10 +288,18 @@ public:
 	 */
 	virtual bool RequiresOriginalReleaseVersionForPatch() const = 0;
 
+	/**
+	 * Checks whether the platform has non-trivial patching requirements, where regular workflows
+	 * may result in invalid patches. It is recommended to check the platform's documentation.
+	 *
+	 * @return true if this platform has non-trivial patching requirements, false otherwise.
+	 */
+	virtual bool HasNonTrivialPatchGeneration() const = 0;
+
 	/**
 	* Checks whether this platform has a secure shippable package format, and therefore doesn't need any encryption or signing support
 	*
-	* @return true if this platform requires cooked data, false otherwise.
+	* @return true if this platform has a secure shippable package format, false otherwise.
 	*/
 	virtual bool HasSecurePackageFormat() const = 0;
 	
diff --git a/Engine/Source/Runtime/Core/Public/GenericPlatform/GenericPlatformProperties.h b/Engine/Source/Runtime/Core/Public/GenericPlatform/GenericPlatformProperties.h
--- a/Engine/Source/Runtime/Core/Public/GenericPlatform/GenericPlatformProperties.h
+++ b/Engine/Source/Runtime/Core/Public/GenericPlatform/GenericPlatformProperties.h
@@ -325,4 +325,10 @@ struct FGenericPlatformProperties
 	{
 		return false;
 	}
+
+	// Whether the platform has non-trivial patching requirements, where regular workflows may result in invalid patches
+	static constexpr UE_FORCEINLINE_HINT bool HasNonTrivialPatchGeneration()
+	{
+		return false;
+	}
 };
