fix for platforms without flavors being greyed out

#rnx
#Synced-CL 48448085

diff --git a/Engine/Plugins/Developer/ProjectLauncher/Source/ProjectLauncher/Private/Widgets/Shared/SCustomLaunchPlatformCombo.cpp b/Engine/Plugins/Developer/ProjectLauncher/Source/ProjectLauncher/Private/Widgets/Shared/SCustomLaunchPlatformCombo.cpp
--- a/Engine/Plugins/Developer/ProjectLauncher/Source/ProjectLauncher/Private/Widgets/Shared/SCustomLaunchPlatformCombo.cpp
+++ b/Engine/Plugins/Developer/ProjectLauncher/Source/ProjectLauncher/Private/Widgets/Shared/SCustomLaunchPlatformCombo.cpp
@@ -220,16 +220,14 @@ FSlateColor SCustomLaunchPlatformCombo::GetPlatformTextColor( const PlatformInfo
 	{
 		bool bHasSupportedBuildTarget;
 
-		if (GetAdvancedPlatformsOption.Get())
+		bHasSupportedBuildTarget = BuildTargets.Contains(PlatformInfo->PlatformType);
+		if (!GetAdvancedPlatformsOption.Get() && !bHasSupportedBuildTarget)
 		{
-			bHasSupportedBuildTarget = BuildTargets.Contains(PlatformInfo->PlatformType);
-		}
-		else
-		{
-			bHasSupportedBuildTarget = PlatformInfo->VanillaInfo->Flavors.ContainsByPredicate( [BuildTargets]( const PlatformInfo::FTargetPlatformInfo* Info )
-			{
-				return BuildTargets.Contains(Info->PlatformType);
-			});
+			bHasSupportedBuildTarget = 
+				PlatformInfo->VanillaInfo->Flavors.ContainsByPredicate( [BuildTargets]( const PlatformInfo::FTargetPlatformInfo* Info )
+				{
+					return BuildTargets.Contains(Info->PlatformType);
+				});
 		}
 
 		if (!bHasSupportedBuildTarget)
