Updates for Project Launcher custom combo widget
- can become a multiselect combo button by setting an ActionType of Check or ToggleButton (use OnSelectedItemsChanged and SelectedItems parameters too)
- new IsItemEnabled, and GetItemIcon callbacks

#rnx

#Synced-CL 48448085

diff --git a/Engine/Plugins/Developer/ProjectLauncher/Source/ProjectLauncher/Public/Widgets/Shared/SCustomLaunchCombo.h b/Engine/Plugins/Developer/ProjectLauncher/Source/ProjectLauncher/Public/Widgets/Shared/SCustomLaunchCombo.h
--- a/Engine/Plugins/Developer/ProjectLauncher/Source/ProjectLauncher/Public/Widgets/Shared/SCustomLaunchCombo.h
+++ b/Engine/Plugins/Developer/ProjectLauncher/Source/ProjectLauncher/Public/Widgets/Shared/SCustomLaunchCombo.h
@@ -10,32 +10,62 @@
 #include "Styling/CoreStyle.h"
 #include "Framework/MultiBox/MultiBoxBuilder.h"
 
-// helper class to create a combo box that can select an item
+#define LOCTEXT_NAMESPACE "SCustomLaunchCombo"
+
+
+// helper class to create a combo box that can select an item or collection of items
 template<typename T>
 class SCustomLaunchCombo
 	: public SCompoundWidget
 {
 public:
 	DECLARE_DELEGATE_OneParam(FOnSelectionChanged, T );
+	DECLARE_DELEGATE_OneParam(FOnSelectedItemsChanged, TArray<T> );
 	DECLARE_DELEGATE_RetVal_OneParam(FText, FGetItemText, T)
+	DECLARE_DELEGATE_RetVal_OneParam(bool, FIsItemEnabled, T)
+	DECLARE_DELEGATE_RetVal_OneParam(FSlateIcon, FGetItemIcon, T)
 
 	SLATE_BEGIN_ARGS(SCustomLaunchCombo<T>)
+		: _ActionType(EUserInterfaceActionType::RadioButton)
 		{}
 		SLATE_EVENT(FOnSelectionChanged, OnSelectionChanged)
 		SLATE_EVENT(FGetItemText, GetDisplayName)
 		SLATE_EVENT(FGetItemText, GetItemToolTip)
+		SLATE_EVENT(FIsItemEnabled, IsItemEnabled)
+		SLATE_EVENT(FGetItemIcon, GetItemIcon)
+		SLATE_EVENT(FOnSelectedItemsChanged, OnSelectedItemsChanged)
 		SLATE_ATTRIBUTE(TArray<T>, Items)
 		SLATE_ATTRIBUTE(T, SelectedItem)
+		SLATE_ATTRIBUTE(TArray<T>, SelectedItems)
+		SLATE_ARGUMENT(EUserInterfaceActionType, ActionType)
 	SLATE_END_ARGS()
 
 public:
 	void Construct(	const FArguments& InArgs)
 	{
 		OnSelectionChanged = InArgs._OnSelectionChanged;
+		OnSelectedItemsChanged = InArgs._OnSelectedItemsChanged;
 		GetDisplayName = InArgs._GetDisplayName;
 		GetItemToolTip = InArgs._GetItemToolTip;
+		IsItemEnabled = InArgs._IsItemEnabled;
+		GetItemIcon = InArgs._GetItemIcon;
+		ActionType = InArgs._ActionType;
+
 		Items = InArgs._Items;
 		SelectedItem = InArgs._SelectedItem;
+		SelectedItems = InArgs._SelectedItems;
+
+		bIsMultiselect = (ActionType == EUserInterfaceActionType::Check || ActionType == EUserInterfaceActionType::ToggleButton);
+		if (bIsMultiselect)
+		{
+			check(!OnSelectionChanged.IsBound());
+			check(!SelectedItem.IsBound());
+		}
+		else
+		{
+			check(!OnSelectedItemsChanged.IsBound());
+			check(!SelectedItems.IsBound());
+		}
 
 		ChildSlot
 		[
@@ -52,20 +82,79 @@ public:
 
 protected:
 	TAttribute<T> SelectedItem;
+	TAttribute<TArray<T>> SelectedItems;
 	TAttribute<TArray<T>> Items;
 	FOnSelectionChanged OnSelectionChanged;
 	FGetItemText GetDisplayName;
 	FGetItemText GetItemToolTip;
+	FIsItemEnabled IsItemEnabled;
+	FGetItemIcon GetItemIcon;
+	FOnSelectedItemsChanged OnSelectedItemsChanged;
+
+	bool bIsMultiselect;
+	EUserInterfaceActionType ActionType;
 
 	FText GetSelectedItemDisplayName() const
 	{
-		T Item = SelectedItem.Get();
-		return ToText(Item);
+		if (bIsMultiselect)
+		{
+			TArray<FText> SelectedItemsText;
+			for (const T& Item : SelectedItems.Get())
+			{
+				SelectedItemsText.Add(ToText(Item));
+			}
+
+			return FText::Join(LOCTEXT("Separator",", "), SelectedItemsText);
+		}
+
+		else
+		{
+			T Item = SelectedItem.Get();
+			return ToText(Item);
+		}
 	}
 
 	void SetSelectedItem(T Value)
 	{
-		OnSelectionChanged.ExecuteIfBound(Value);
+		if (bIsMultiselect)
+		{
+			TArray<T> Values = SelectedItems.Get();
+			if (Values.Contains(Value))
+			{
+				Values.Remove(Value);
+			}
+			else
+			{
+				Values.Add(Value);
+			}
+
+			OnSelectedItemsChanged.ExecuteIfBound(Values);
+
+		}
+		else
+		{
+			OnSelectionChanged.ExecuteIfBound(Value);
+		}
+	}
+
+	ECheckBoxState GetItemCheckState(T Value)
+	{
+		bool bResult;
+		if (bIsMultiselect)
+		{
+			bResult = SelectedItems.Get().Contains(Value);
+		}
+		else
+		{
+			bResult = SelectedItem.Get() == Value;
+		}
+
+		return bResult ? ECheckBoxState::Checked : ECheckBoxState::Unchecked; 
+	}
+
+	bool GetItemIsEnabled(T Value) const
+	{
+		return IsItemEnabled.IsBound() ? IsItemEnabled.Execute(Value) : true;
 	}
 
 	FText ToText(T Value) const
@@ -84,11 +173,11 @@ protected:
 		return FText::FromString("ERROR: GetDisplayName or ToTextOverride not bound");
 	}
 
-
 BEGIN_SLATE_FUNCTION_BUILD_OPTIMIZATION
 	TSharedRef<SWidget> MakeWidget()
 	{
-		FMenuBuilder MenuBuilder(true, nullptr);
+		const bool bCloseOnSelection = !bIsMultiselect;
+		FMenuBuilder MenuBuilder(bCloseOnSelection, nullptr);
 		{
 			TArray<T> MenuItems = Items.Get();
 
@@ -96,18 +185,19 @@ BEGIN_SLATE_FUNCTION_BUILD_OPTIMIZATION
 			{
 				FText ItemText = ToText(Item);
 				FText ToolTipText = GetItemToolTip.IsBound() ? GetItemToolTip.Execute(Item) : FText::GetEmpty();
+				FSlateIcon Icon = GetItemIcon.IsBound() ? GetItemIcon.Execute(Item) : FSlateIcon();
 
 				MenuBuilder.AddMenuEntry(
 					ItemText, 
 					ToolTipText, 
-					FSlateIcon(), 
+					Icon, 
 					FUIAction(
 						FExecuteAction::CreateSP(this, &SCustomLaunchCombo<T>::SetSelectedItem, Item),
-						FCanExecuteAction(),
-						FGetActionCheckState::CreateLambda( [this, Item]() { return (SelectedItem.Get() == Item) ? ECheckBoxState::Checked : ECheckBoxState::Unchecked; })
+						FCanExecuteAction::CreateSP(this, &SCustomLaunchCombo<T>::GetItemIsEnabled, Item),
+						FGetActionCheckState::CreateSP( this, &SCustomLaunchCombo<T>::GetItemCheckState, Item)
 					),
 					NAME_None,
-					EUserInterfaceActionType::Check);
+					ActionType);
 			}
 		}
 
@@ -141,4 +231,5 @@ public:
 };
 
 
+#undef LOCTEXT_NAMESPACE
 
