Fix warnings for missing resources
- The core content root is set to an Engine path, not Editor

#Synced-CL 48571193

diff --git a/Engine/Plugins/Developer/ProjectLauncher/Source/ProjectLauncher/Private/Styling/ProjectLauncherStyle.cpp b/Engine/Plugins/Developer/ProjectLauncher/Source/ProjectLauncher/Private/Styling/ProjectLauncherStyle.cpp
--- a/Engine/Plugins/Developer/ProjectLauncher/Source/ProjectLauncher/Private/Styling/ProjectLauncherStyle.cpp
+++ b/Engine/Plugins/Developer/ProjectLauncher/Source/ProjectLauncher/Private/Styling/ProjectLauncherStyle.cpp
@@ -64,18 +64,18 @@ void FProjectLauncherStyle::Initialize()
 	StyleSet->Set("State.AllComplete", FLinearColor(FColor(0xFF0070E0)) );
 
 	FButtonStyle HoverHintOnly = FButtonStyle()
-			.SetNormal( FSlateNoResource() )
-			.SetHovered( CORE_BOX_BRUSH( "Common/ButtonHoverHint", FMargin(4.f/16.0f), FLinearColor(1.f,1.f,1.f,0.15f) ) )
-			.SetPressed( CORE_BOX_BRUSH( "Common/ButtonHoverHint", FMargin(4.f/16.0f), FLinearColor(1.f,1.f,1.f,0.25f) ) )
-			.SetNormalPadding( FMargin(0.f,0.f,0.f,1.f) )
-			.SetPressedPadding( FMargin(0.f,1.f,0.f,0.f) );
+		.SetNormal(FSlateNoResource())
+		.SetHovered(CORE_BOX_BRUSH("../Editor/Slate/Common/ButtonHoverHint", FMargin(4.f/16.0f), FLinearColor(1.f,1.f,1.f,0.15f)))
+		.SetPressed(CORE_BOX_BRUSH("../Editor/Slate/Common/ButtonHoverHint", FMargin(4.f/16.0f), FLinearColor(1.f,1.f,1.f,0.25f)))
+		.SetNormalPadding(FMargin(0.f,0.f,0.f,1.f))
+		.SetPressedPadding(FMargin(0.f,1.f,0.f,0.f));
 	StyleSet->Set( "HoverHintOnly", HoverHintOnly );
 
 	StyleSet->Set("Icons.DiffersFromDefault", new CORE_IMAGE_BRUSH_SVG( "../Editor/Slate/Starship/Common/ResetToDefault", Icon16x16) );
 	StyleSet->Set("Icons.DeviceManager", new CORE_IMAGE_BRUSH_SVG( "../Editor/Slate/Starship/Common/DeviceManager", Icon16x16) );
 	StyleSet->Set("Icons.Plugins", new CORE_IMAGE_BRUSH_SVG("../Editor/Slate/Starship/Common/Plugins", Icon16x16));
 
-	StyleSet->Set("Icons.Asset", new CORE_IMAGE_BRUSH( "./Editor/Slate/Icons/doc_16x", Icon16x16 ) );
+	StyleSet->Set("Icons.Asset", new CORE_IMAGE_BRUSH( "../Editor/Slate/Icons/doc_16x", Icon16x16 ) );
 	StyleSet->Set("Icons.WarningWithColor.Small", new CORE_IMAGE_BRUSH_SVG("../Editor/Slate/Starship/Common/AlertTriangleSolid", Icon16x16, FStyleColors::Warning));
 
 	StyleSet->Set( "PathPickerButton", new CORE_IMAGE_BRUSH( "Icons/ellipsis_12x", Icon12x12 ) );
