Deprecate the FArchive serialization interface for launcher profiles in favor of the Json one

#rnx
#Synced-CL 48448085

diff --git a/Engine/Plugins/Developer/ProjectLauncher/Source/ProjectLauncher/Private/Model/ProjectLauncherModel.cpp b/Engine/Plugins/Developer/ProjectLauncher/Source/ProjectLauncher/Private/Model/ProjectLauncherModel.cpp
--- a/Engine/Plugins/Developer/ProjectLauncher/Source/ProjectLauncher/Private/Model/ProjectLauncherModel.cpp
+++ b/Engine/Plugins/Developer/ProjectLauncher/Source/ProjectLauncher/Private/Model/ProjectLauncherModel.cpp
@@ -11,8 +11,8 @@
 #include "ITargetDeviceProxy.h"
 #include "ITargetDeviceProxyManager.h"
 #include "ITargetDeviceServicesModule.h"
-#include "Serialization/BufferArchive.h"
-#include "Serialization/MemoryReader.h"
+#include "Serialization/JsonReader.h"
+#include "Serialization/JsonSerializer.h"
 #include "Modules/ModuleManager.h"
 #include "Settings/ProjectPackagingSettings.h"
 #include "Experimental/ZenServerInterface.h"
@@ -631,32 +631,37 @@ PRAGMA_ENABLE_DEPRECATION_WARNINGS
 
 	ILauncherProfilePtr FModel::CloneCustomProfile(const ILauncherProfileRef& Profile)
 	{
-		ILauncherProfilePtr NewProfile;
-
-		FBufferArchive Writer;
-		if (Profile->Serialize(Writer))
-		{
-			NewProfile = CreateCustomProfile(TEXT("Cloned"));
-			if (NewProfile.IsValid())
+		// capture the current profile
+		FString ExistingProfileJson;
+		TSharedRef<TJsonWriter<>> JsonWriter = TJsonWriterFactory<>::Create(&ExistingProfileJson);
+		Profile->Save(JsonWriter.Get());
+		JsonWriter->Close();
+
+		// create a new json object
+		TSharedPtr<FJsonObject> ProfileJsonObject;
+		TSharedRef<TJsonReader<>> JsonReader = TJsonReaderFactory<>::Create(ExistingProfileJson);
+		if (FJsonSerializer::Deserialize(JsonReader, ProfileJsonObject) && ProfileJsonObject.IsValid())
+		{
+			// create the cloned profile & load the json
+			ILauncherProfilePtr NewProfile = CreateCustomProfile(TEXT("Cloned"));
+			if (NewProfile.IsValid() && NewProfile->Load(*ProfileJsonObject.Get()))
 			{
-				FMemoryReader Reader(Writer);
-				NewProfile->Serialize(Reader);
-				NewProfile->AssignId(true); // need to give it a new Id - don't use the cloned one!
-
-				// @todo: add these to FLauncherProfile's Serialize() function & update the internal version. Waiting for all Zen properties to be decided (zen workspace etc)
-				NewProfile->SetUseZenStore(Profile->IsUsingZenStore());
-				// ...
+				// need to give it a new Id - don't use the cloned one!
+				NewProfile->AssignId(true);
 
+				// create a new device group because these are not serialized
 				ILauncherDeviceGroupRef DeployDeviceGroup = ProfileManager->AddNewDeviceGroup();
 				NewProfile->SetDeployedDeviceGroup(DeployDeviceGroup);
 				for (const FString& DeviceID : Profile->GetDeployedDeviceGroup()->GetDeviceIDs())
 				{
 					DeployDeviceGroup->AddDevice(DeviceID);
 				}
+
+				return NewProfile;
 			}
 		}
 
-		return NewProfile;
+		return nullptr;
 	}
 
 			
diff --git a/Engine/Plugins/Developer/ProjectLauncher/Source/ProjectLauncher/ProjectLauncher.Build.cs b/Engine/Plugins/Developer/ProjectLauncher/Source/ProjectLauncher/ProjectLauncher.Build.cs
--- a/Engine/Plugins/Developer/ProjectLauncher/Source/ProjectLauncher/ProjectLauncher.Build.cs
+++ b/Engine/Plugins/Developer/ProjectLauncher/Source/ProjectLauncher/ProjectLauncher.Build.cs
@@ -20,6 +20,7 @@ public class ProjectLauncher : ModuleRules
 			"Projects",
 			"DeveloperToolSettings",
 			"Zen",
+			"Json",
 		]);
 
 		PublicDependencyModuleNames.AddRange(
diff --git a/Engine/Source/Developer/LauncherServices/Private/Profiles/LauncherProfile.h b/Engine/Source/Developer/LauncherServices/Private/Profiles/LauncherProfile.h
--- a/Engine/Source/Developer/LauncherServices/Private/Profiles/LauncherProfile.h
+++ b/Engine/Source/Developer/LauncherServices/Private/Profiles/LauncherProfile.h
@@ -155,6 +155,7 @@ public:
 		CookMode = InMode;
 	}
 
+	UE_DEPRECATED(5.8, "this will be removed. there is no alternative implementation : ILauncherSimpleProfile is designed to be transient")
 	virtual bool Serialize(FArchive& Archive) override
 	{
 		int32 Version = LAUNCHERSERVICES_SIMPLEPROFILEVERSION;
@@ -1024,6 +1025,7 @@ public:
 		Validate();
 	}
 
+	UE_DEPRECATED(5.8, "please use the Json Save/Load API instead")
 	virtual bool Serialize( FArchive& Archive ) override
 	{
 		int32 Version = LAUNCHERSERVICES_FINAL;
@@ -1242,6 +1244,7 @@ public:
 			Archive << NumAutomatedTests;
 			for( int32 TestIndex = 0; TestIndex < NumAutomatedTests; TestIndex++)
 			{
+PRAGMA_DISABLE_DEPRECATION_WARNINGS
 				if (Archive.IsLoading())
 				{
 					AutomatedTests.Add(MakeShared<FLauncherProfileAutomatedTest>(Archive));
@@ -1251,6 +1254,7 @@ public:
 					AutomatedTests[TestIndex]->Serialize(Archive);
 				}
 			}
+PRAGMA_ENABLE_DEPRECATION_WARNINGS
 
 			Archive << AutomatedTestBuildPath;
 			Archive << bWantAutomatedTestBuild;
@@ -1267,7 +1271,9 @@ public:
 			Archive << AdditionalTargetCommandLineParameters;
 		}
 
+PRAGMA_DISABLE_DEPRECATION_WARNINGS
 		DefaultLaunchRole->Serialize(Archive);
+PRAGMA_ENABLE_DEPRECATION_WARNINGS
 
 		// serialize launch roles
 		if (Archive.IsLoading())
@@ -1282,6 +1288,7 @@ public:
 
 		for (int32 RoleIndex = 0; RoleIndex < NumLaunchRoles; ++RoleIndex)
 		{
+PRAGMA_DISABLE_DEPRECATION_WARNINGS
 			if (Archive.IsLoading())
 			{
 				LaunchRoles.Add(MakeShareable(new FLauncherProfileLaunchRole(Archive)));				
@@ -1290,6 +1297,7 @@ public:
 			{
 				LaunchRoles[RoleIndex]->Serialize(Archive);
 			}
+PRAGMA_ENABLE_DEPRECATION_WARNINGS
 		}
 
 		if (Archive.IsLoading())
diff --git a/Engine/Source/Developer/LauncherServices/Private/Profiles/LauncherProfileAutomatedTest.h b/Engine/Source/Developer/LauncherServices/Private/Profiles/LauncherProfileAutomatedTest.h
--- a/Engine/Source/Developer/LauncherServices/Private/Profiles/LauncherProfileAutomatedTest.h
+++ b/Engine/Source/Developer/LauncherServices/Private/Profiles/LauncherProfileAutomatedTest.h
@@ -19,9 +19,12 @@ public:
 		Load(Object);
 	}
 
+	//UE_DEPRECATED(5.8, "please use the Json Save/Load API") - confuses linux build which thinks the entire class is deprecated
 	FLauncherProfileAutomatedTest( FArchive& Archive )
 	{
+		PRAGMA_DISABLE_DEPRECATION_WARNINGS
 		Serialize(Archive);
+		PRAGMA_ENABLE_DEPRECATION_WARNINGS
 	}
 
 	virtual void SetUATCommand( const TCHAR* InUATCommand ) override
@@ -97,6 +100,7 @@ public:
 		Writer.WriteObjectEnd();
 	}
 
+	UE_DEPRECATED(5.8, "please use the Json Save/Load API")
 	virtual void Serialize( FArchive& Archive ) override final
 	{
 		Archive << InternalName
diff --git a/Engine/Source/Developer/LauncherServices/Private/Profiles/LauncherProfileLaunchRole.h b/Engine/Source/Developer/LauncherServices/Private/Profiles/LauncherProfileLaunchRole.h
--- a/Engine/Source/Developer/LauncherServices/Private/Profiles/LauncherProfileLaunchRole.h
+++ b/Engine/Source/Developer/LauncherServices/Private/Profiles/LauncherProfileLaunchRole.h
@@ -22,9 +22,12 @@ public:
 	 *
 	 * @param Archive - The archive to serialize from.
 	 */
+	UE_DEPRECATED(5.8, "please use the Json Save/Load API")
 	FLauncherProfileLaunchRole( FArchive& Archive )
 	{
+		PRAGMA_DISABLE_DEPRECATION_WARNINGS
 		Serialize(Archive);
+		PRAGMA_ENABLE_DEPRECATION_WARNINGS
 	}
 
 	FLauncherProfileLaunchRole(const FJsonObject& Object)
@@ -97,6 +100,7 @@ public:
 		Writer.WriteObjectEnd();
 	}
 
+	UE_DEPRECATED(5.8, "please use the Json Save/Load API")
 	virtual void Serialize( FArchive& Archive ) override final
 	{
 		Archive << AssignedDevice
diff --git a/Engine/Source/Developer/LauncherServices/Private/Profiles/LauncherProfileManager.cpp b/Engine/Source/Developer/LauncherServices/Private/Profiles/LauncherProfileManager.cpp
--- a/Engine/Source/Developer/LauncherServices/Private/Profiles/LauncherProfileManager.cpp
+++ b/Engine/Source/Developer/LauncherServices/Private/Profiles/LauncherProfileManager.cpp
@@ -262,7 +262,9 @@ ILauncherProfilePtr FLauncherProfileManager::LoadProfile( FArchive& Archive )
 {
 	FLauncherProfile* Profile = new FLauncherProfile(AsShared());
 
+PRAGMA_DISABLE_DEPRECATION_WARNINGS
 	if (Profile->Serialize(Archive))
+PRAGMA_ENABLE_DEPRECATION_WARNINGS
 	{
 		ILauncherDeviceGroupPtr DeviceGroup = GetDeviceGroup(Profile->GetDeployedDeviceGroupId());
 		if (!DeviceGroup.IsValid())
@@ -368,7 +370,9 @@ bool FLauncherProfileManager::SaveProfile(const ILauncherProfileRef& Profile)
 
 		if (ProfileFileWriter != nullptr)
 		{
+PRAGMA_DISABLE_DEPRECATION_WARNINGS
 			Profile->Serialize(*ProfileFileWriter);
+PRAGMA_ENABLE_DEPRECATION_WARNINGS
 
 			delete ProfileFileWriter;
 
@@ -537,6 +541,7 @@ void FLauncherProfileManager::LoadProfiles( )
 
 			if (ProfileFileReader != nullptr)
 			{
+PRAGMA_DISABLE_DEPRECATION_WARNINGS
 				ILauncherProfilePtr LoadedProfile = LoadProfile(*ProfileFileReader);
 				delete ProfileFileReader;
 
@@ -545,6 +550,7 @@ void FLauncherProfileManager::LoadProfiles( )
 				{
 					SaveProfile(LoadedProfile.ToSharedRef());
 				}
+PRAGMA_ENABLE_DEPRECATION_WARNINGS
 
 				//delete legacy profile.
 				IFileManager::Get().Delete(*ProfileFilePath);				
@@ -564,8 +570,10 @@ void FLauncherProfileManager::LoadProfiles( )
 
 			if (ProfileFileReader != nullptr)
 			{
+PRAGMA_DISABLE_DEPRECATION_WARNINGS
 				ILauncherProfilePtr LoadedProfile = LoadProfile(*ProfileFileReader);
 				delete ProfileFileReader;
+PRAGMA_ENABLE_DEPRECATION_WARNINGS
 
 				//re-save profile to the new format
 				if (LoadedProfile.IsValid())
diff --git a/Engine/Source/Developer/LauncherServices/Private/Profiles/LauncherProfileManager.h b/Engine/Source/Developer/LauncherServices/Private/Profiles/LauncherProfileManager.h
--- a/Engine/Source/Developer/LauncherServices/Private/Profiles/LauncherProfileManager.h
+++ b/Engine/Source/Developer/LauncherServices/Private/Profiles/LauncherProfileManager.h
@@ -54,6 +54,7 @@ public:
 
 	virtual ILauncherProfilePtr GetProfile( const FGuid& ProfileId ) const override;
 
+	UE_DEPRECATED(5.8, "Please use LoadJSONProfile")
 	virtual ILauncherProfilePtr LoadProfile( FArchive& Archive ) override;
 
 	virtual ILauncherProfilePtr LoadJSONProfile(FString ProfileFile) override;
@@ -103,6 +104,7 @@ public:
 
 	virtual void RemoveProfile( const ILauncherProfileRef& Profile ) override;
 
+	UE_DEPRECATED(5.8, "Please use SaveJSONProfile")
 	virtual bool SaveProfile( const ILauncherProfileRef& Profile) override;
 
 	virtual bool SaveJSONProfile(const ILauncherProfileRef& Profile) override;
diff --git a/Engine/Source/Developer/LauncherServices/Public/ILauncherProfile.h b/Engine/Source/Developer/LauncherServices/Public/ILauncherProfile.h
--- a/Engine/Source/Developer/LauncherServices/Public/ILauncherProfile.h
+++ b/Engine/Source/Developer/LauncherServices/Public/ILauncherProfile.h
@@ -365,6 +365,7 @@ public:
 	 * @param Archive The archive to serialize from or into.
 	 * @return true if the profile was serialized, false otherwise.
 	 */
+	UE_DEPRECATED(5.8, "this will be removed. there is no alternative implementation : ILauncherSimpleProfile is designed to be transient")
 	virtual bool Serialize(FArchive& Archive) = 0;
 
 	/** Sets all profile settings to their defaults. */
@@ -579,6 +580,7 @@ public:
 	 * @param Archive The archive to serialize from or into.
 	 * @return true if the profile was serialized, false otherwise.
 	 */
+	UE_DEPRECATED(5.8, "please use the Json Save/Load API instead")
 	virtual bool Serialize( FArchive& Archive ) = 0;
 
 	/**
diff --git a/Engine/Source/Developer/LauncherServices/Public/ILauncherProfileAutomatedTest.h b/Engine/Source/Developer/LauncherServices/Public/ILauncherProfileAutomatedTest.h
--- a/Engine/Source/Developer/LauncherServices/Public/ILauncherProfileAutomatedTest.h
+++ b/Engine/Source/Developer/LauncherServices/Public/ILauncherProfileAutomatedTest.h
@@ -36,6 +36,8 @@ public:
 
 	virtual void Load(const FJsonObject& Object) = 0;
 	virtual void Save(TJsonWriter<>& Writer) = 0;
+
+	UE_DEPRECATED(5.8, "please use the Json Save/Load API")
 	virtual void Serialize( FArchive& Archive ) = 0;
 
 
diff --git a/Engine/Source/Developer/LauncherServices/Public/ILauncherProfileLaunchRole.h b/Engine/Source/Developer/LauncherServices/Public/ILauncherProfileLaunchRole.h
--- a/Engine/Source/Developer/LauncherServices/Public/ILauncherProfileLaunchRole.h
+++ b/Engine/Source/Developer/LauncherServices/Public/ILauncherProfileLaunchRole.h
@@ -134,6 +134,7 @@ public:
 	 *
 	 * @param Archive The archive to serialize from or into.
 	 */
+	UE_DEPRECATED(5.8, "please use the Json Save/Load API")
 	virtual void Serialize( FArchive& Archive ) = 0;
 
 	/**
diff --git a/Engine/Source/Developer/LauncherServices/Public/ILauncherProfileManager.h b/Engine/Source/Developer/LauncherServices/Public/ILauncherProfileManager.h
--- a/Engine/Source/Developer/LauncherServices/Public/ILauncherProfileManager.h
+++ b/Engine/Source/Developer/LauncherServices/Public/ILauncherProfileManager.h
@@ -241,6 +241,7 @@ public:
 	 * @return The loaded profile, or nullptr if loading failed.
 	 * @see AddProfile, SaveProfile
 	 */
+	UE_DEPRECATED(5.8, "Please use LoadJSONProfile")
 	virtual ILauncherProfilePtr LoadProfile( FArchive& Archive ) = 0;
 
 	/**
@@ -268,6 +269,7 @@ public:
 	 * @param Archive The archive to save to.
 	 * @see LoadProfile
 	 */
+	UE_DEPRECATED(5.8, "Please use SaveJSONProfile")
 	virtual bool SaveProfile(const ILauncherProfileRef& Profile) = 0;
 
 	/**
