hook profile creation wizards into the extensions menu, under 'deprecated'
- likely unused now so don't want to expand the normal profile creation button

#rnx
#Synced-CL 48448085

diff --git a/Engine/Plugins/Developer/ProjectLauncher/Source/CommonLaunchExtensions/Private/CommonLaunchExtensionsModule.cpp b/Engine/Plugins/Developer/ProjectLauncher/Source/CommonLaunchExtensions/Private/CommonLaunchExtensionsModule.cpp
--- a/Engine/Plugins/Developer/ProjectLauncher/Source/CommonLaunchExtensions/Private/CommonLaunchExtensionsModule.cpp
+++ b/Engine/Plugins/Developer/ProjectLauncher/Source/CommonLaunchExtensions/Private/CommonLaunchExtensionsModule.cpp
@@ -7,6 +7,7 @@
 #include "Insights/InsightsLaunchExtension.h"
 #include "BootTest/BootTestLaunchExtension.h"
 #include "AdvancedCook/AdvancedCookLaunchExtension.h"
+#include "ProfileWizard/ProfileWizardLaunchExtension.h"
 
 class FCommonLaunchExtensionsModule : public IModuleInterface
 {
@@ -18,6 +19,7 @@ public:
 		Extensions.Add( MakeShared<FInsightsLaunchExtension>() );
 		Extensions.Add( MakeShared<FBootTestLaunchExtension>() );
 		Extensions.Add( MakeShared<FAdvancedCookLaunchExtension>() );
+		Extensions.Add( MakeShared<FProfileWizardLaunchExtension>() );
 
 		for (TSharedRef<ProjectLauncher::FLaunchExtension> Extension : Extensions)
 		{
diff --git a/Engine/Plugins/Developer/ProjectLauncher/Source/CommonLaunchExtensions/Private/ProfileWizard/ProfileWizardLaunchExtension.cpp b/Engine/Plugins/Developer/ProjectLauncher/Source/CommonLaunchExtensions/Private/ProfileWizard/ProfileWizardLaunchExtension.cpp
new file mode 100644
--- /dev/null
+++ b/Engine/Plugins/Developer/ProjectLauncher/Source/CommonLaunchExtensions/Private/ProfileWizard/ProfileWizardLaunchExtension.cpp
@@ -0,0 +1,74 @@
+// Copyright Epic Games, Inc. All Rights Reserved.
+
+#include "ProfileWizard/ProfileWizardLaunchExtension.h"
+#include "Model/ProjectLauncherModel.h"
+
+#define LOCTEXT_NAMESPACE "FProfileWizardLaunchExtensionInstance"
+
+
+FProfileWizardLaunchExtensionInstance::FProfileWizardLaunchExtensionInstance( FArgs& InArgs ) 
+	: ProjectLauncher::FLaunchExtensionInstance(InArgs)
+{
+}
+
+
+void FProfileWizardLaunchExtensionInstance::MakeCustomExtensionSubmenu( FMenuBuilder& MenuBuilder )
+{
+	auto CreateSubMenu = [this]( FMenuBuilder& MenuBuilder )
+	{
+		for (const ILauncherProfileWizardPtr& ProfileWizard : GetModel()->GetProfileManager()->GetProfileWizards())
+		{
+			auto RunWizard = [this, ProfileWizard]()
+			{
+				ProfileWizard->HandleCreateLauncherProfile( GetModel()->GetProfileManager() );
+			};
+
+			MenuBuilder.AddMenuEntry(
+				ProfileWizard->GetName(),
+				ProfileWizard->GetDescription(),
+				FSlateIcon(),
+				FUIAction( FExecuteAction::CreateLambda(RunWizard) ),
+				NAME_None);
+		}
+	};
+
+	MenuBuilder.AddSubMenu( GetExtension()->GetDisplayName(), FText::GetEmpty(), FNewMenuDelegate::CreateLambda(CreateSubMenu));
+}
+
+
+void FProfileWizardLaunchExtensionInstance::GetExtensionsMenuEntry( FExtensionsMenuEntry& MenuEntry ) const
+{
+	MenuEntry.SectionName = FString();
+
+	MenuEntry.Type = FExtensionsMenuEntry::Type_SubMenu;
+	MenuEntry.SubmenuName = TEXT("Deprecated");
+	MenuEntry.SubmenuDisplayName = LOCTEXT("DeprecatedSubmenu", "Deprecated");
+}
+
+
+
+TSharedPtr<ProjectLauncher::FLaunchExtensionInstance> FProfileWizardLaunchExtension::CreateInstanceForProfile( ProjectLauncher::FLaunchExtensionInstance::FArgs& InArgs )
+{
+	if (InArgs.Model->GetProfileManager()->GetProfileWizards().Num() > 0)
+	{
+		return MakeShared<FProfileWizardLaunchExtensionInstance>(InArgs);
+	}
+
+	// there's no profile wizards - don't need to create the extension
+	return nullptr;
+}
+
+
+const TCHAR* FProfileWizardLaunchExtension::GetInternalName() const
+{
+	return TEXT("ProfileWizard");
+}
+
+
+FText FProfileWizardLaunchExtension::GetDisplayName() const
+{
+	return LOCTEXT("ExtensionName", "New Profile Wizard");
+}
+
+
+#undef LOCTEXT_NAMESPACE
diff --git a/Engine/Plugins/Developer/ProjectLauncher/Source/CommonLaunchExtensions/Private/ProfileWizard/ProfileWizardLaunchExtension.h b/Engine/Plugins/Developer/ProjectLauncher/Source/CommonLaunchExtensions/Private/ProfileWizard/ProfileWizardLaunchExtension.h
new file mode 100644
--- /dev/null
+++ b/Engine/Plugins/Developer/ProjectLauncher/Source/CommonLaunchExtensions/Private/ProfileWizard/ProfileWizardLaunchExtension.h
@@ -0,0 +1,27 @@
+// Copyright Epic Games, Inc. All Rights Reserved.
+
+#pragma once
+
+#include "Extension/LaunchExtension.h"
+
+class FProfileWizardLaunchExtensionInstance : public ProjectLauncher::FLaunchExtensionInstance
+{
+public:
+	FProfileWizardLaunchExtensionInstance( FArgs& InArgs );
+	virtual ~FProfileWizardLaunchExtensionInstance() = default;
+
+	virtual bool HasCustomExtensionMenu() const { return true; }
+	virtual void MakeCustomExtensionSubmenu( FMenuBuilder& MenuBuilder ) override;
+	virtual void GetExtensionsMenuEntry( FExtensionsMenuEntry& MenuEntry ) const override;
+
+protected:
+};
+
+
+class FProfileWizardLaunchExtension : public ProjectLauncher::FLaunchExtension
+{
+public:
+	virtual TSharedPtr<ProjectLauncher::FLaunchExtensionInstance> CreateInstanceForProfile( ProjectLauncher::FLaunchExtensionInstance::FArgs& InArgs ) override;
+	virtual const TCHAR* GetInternalName() const override;
+	virtual FText GetDisplayName() const override;
+};
\ No newline at end of file
