Add all remaining legacy project launcher properties to an extension

#rnx
#Synced-CL 48448085

diff --git a/Engine/Plugins/Developer/ProjectLauncher/Source/CommonLaunchExtensions/Private/CommonLaunchExtensionsModule.cpp b/Engine/Plugins/Developer/ProjectLauncher/Source/CommonLaunchExtensions/Private/CommonLaunchExtensionsModule.cpp
--- a/Engine/Plugins/Developer/ProjectLauncher/Source/CommonLaunchExtensions/Private/CommonLaunchExtensionsModule.cpp
+++ b/Engine/Plugins/Developer/ProjectLauncher/Source/CommonLaunchExtensions/Private/CommonLaunchExtensionsModule.cpp
@@ -8,6 +8,7 @@
 #include "BootTest/BootTestLaunchExtension.h"
 #include "AdvancedCook/AdvancedCookLaunchExtension.h"
 #include "ProfileWizard/ProfileWizardLaunchExtension.h"
+#include "DeprecatedProperties/DeprecatedPropertiesLaunchExtension.h"
 
 class FCommonLaunchExtensionsModule : public IModuleInterface
 {
@@ -20,6 +21,7 @@ public:
 		Extensions.Add( MakeShared<FBootTestLaunchExtension>() );
 		Extensions.Add( MakeShared<FAdvancedCookLaunchExtension>() );
 		Extensions.Add( MakeShared<FProfileWizardLaunchExtension>() );
+		Extensions.Add( MakeShared<FDeprecatedPropertiesLaunchExtension>() );
 
 		for (TSharedRef<ProjectLauncher::FLaunchExtension> Extension : Extensions)
 		{
diff --git a/Engine/Plugins/Developer/ProjectLauncher/Source/CommonLaunchExtensions/Private/DeprecatedProperties/DeprecatedPropertiesLaunchExtension.cpp b/Engine/Plugins/Developer/ProjectLauncher/Source/CommonLaunchExtensions/Private/DeprecatedProperties/DeprecatedPropertiesLaunchExtension.cpp
new file mode 100644
--- /dev/null
+++ b/Engine/Plugins/Developer/ProjectLauncher/Source/CommonLaunchExtensions/Private/DeprecatedProperties/DeprecatedPropertiesLaunchExtension.cpp
@@ -0,0 +1,211 @@
+// Copyright Epic Games, Inc. All Rights Reserved.
+
+#include "DeprecatedProperties/DeprecatedPropertiesLaunchExtension.h"
+#include "HAL/PlatformMisc.h"
+#include "Misc/CoreMisc.h"
+#include "Widgets/Shared/SCustomLaunchCombo.h"
+#include "Widgets/Text/STextBlock.h"
+#include "Widgets/Layout/SBorder.h"
+#include "Styling/AppStyle.h"
+
+#define LOCTEXT_NAMESPACE "FDeprecatedPropertiesLaunchExtensionInstance"
+
+
+
+FDeprecatedPropertiesLaunchExtensionInstance::FDeprecatedPropertiesLaunchExtensionInstance( FArgs& InArgs ) 
+	: ProjectLauncher::FLaunchExtensionInstance(InArgs)
+{
+	bEnabled = HasDeprecatedProperties();
+}
+
+
+bool FDeprecatedPropertiesLaunchExtensionInstance::HasDeprecatedProperties() const
+{
+	return	GetProfile()->GetSkipCookingEditorContent()				!=	GetDefaultProfile()->GetSkipCookingEditorContent() ||
+			GetProfile()->GetCookConfiguration()					!=	GetDefaultProfile()->GetCookConfiguration() ||
+			GetProfile()->GetPackageDirectory()						!=	GetDefaultProfile()->GetPackageDirectory() ||
+			GetProfile()->IsUsingIoStore()							!=	GetDefaultProfile()->IsUsingIoStore() ||
+			GetProfile()->GetReferenceContainerGlobalFileName()		!=	GetDefaultProfile()->GetReferenceContainerGlobalFileName() ||
+			GetProfile()->GetReferenceContainerCryptoKeysFileName()	!=	GetDefaultProfile()->GetReferenceContainerCryptoKeysFileName()
+		;
+}
+
+void FDeprecatedPropertiesLaunchExtensionInstance::CustomizeTree( ProjectLauncher::FLaunchProfileTreeData& ProfileTreeData )
+{
+	auto IsVisible = [this]()
+	{
+		return bEnabled;
+	};
+
+	AddDefaultHeading(ProfileTreeData)
+		.AddWidget( FText::GetEmpty(),
+			{
+				.IsVisible = IsVisible,
+			},
+			SNew(SBorder)
+			.BorderImage(FAppStyle::GetBrush("ToolPanel.LightGroupBorder"))
+			.BorderBackgroundColor_Lambda( [this]()
+			{
+				return HasDeprecatedProperties()
+					? FAppStyle::Get().GetSlateColor("Colors.Warning")
+					: FAppStyle::Get().GetSlateColor("Colors.AccentGreen");
+			})
+			[
+				SNew(STextBlock)
+				.ColorAndOpacity(FLinearColor::White)
+				.Text_Lambda( [this]() 
+				{ 
+					return HasDeprecatedProperties() 
+						? LOCTEXT("DeprecatedItemsBanner", "These properties may be removed in future engine releases.") 
+						: LOCTEXT("NoDeprecatedItemsBanner", "These properties may be removed in future engine releases. You are not using any at the moment.");
+				})
+			]
+		)
+
+		// ... http chunks ...
+		.AddBoolean( LOCTEXT("CreateHttpChunksLabel", "Create Http Chunk Install Data"),
+			{
+				.GetValue = [this]()				{ return GetProfile()->IsGenerateHttpChunkData(); },
+				.SetValue = [this](bool bValue)		{ GetProfile()->SetGenerateHttpChunkData(bValue); },
+				.GetDefaultValue = [this]()			{ return GetDefaultProfile()->IsGenerateHttpChunkData(); },
+				.IsVisible = IsVisible,
+			}
+		)
+		.AddDirectoryString( LOCTEXT("HttpChunkDirLabel", "Http Chunk Install Data Path"),
+			{
+				.GetValue = [this]()				{ return GetProfile()->GetHttpChunkDataDirectory(); },
+				.SetValue = [this](FString Value)	{ GetProfile()->SetHttpChunkDataDirectory(Value); },
+				.GetDefaultValue = [this]()			{ return GetDefaultProfile()->GetHttpChunkDataDirectory(); },
+				.IsVisible = IsVisible,
+				.IsEnabled = [this]()				{ return GetProfile()->IsGenerateHttpChunkData(); },
+			}
+		)
+		.AddString( LOCTEXT("HttpChunkReleaseName", "Http Chunk Install Release Name"),
+			{
+				.GetValue = [this]()				{ return GetProfile()->GetHttpChunkDataReleaseName(); },
+				.SetValue = [this](FString Value)	{ GetProfile()->SetHttpChunkDataReleaseName(Value); },
+				.GetDefaultValue = [this]()			{ return GetDefaultProfile()->GetHttpChunkDataReleaseName(); },
+				.IsVisible = IsVisible,
+				.IsEnabled = [this]()				{ return GetProfile()->IsGenerateHttpChunkData(); },
+			}
+		)
+
+		// ... misc items ...
+		.AddBoolean( LOCTEXT("SkipEditorContentLabel", "Don't include editor content in the build"),
+			{
+				.GetValue = [this]()				{ return GetProfile()->GetSkipCookingEditorContent(); },
+				.SetValue = [this](bool bValue)		{ GetProfile()->SetSkipCookingEditorContent(bValue); },
+				.GetDefaultValue = [this]()			{ return GetDefaultProfile()->GetSkipCookingEditorContent(); },
+				.IsVisible = IsVisible,
+			}
+		)
+		.AddWidget( LOCTEXT("CookerConfigLabel", "Cooker build configuration"), 
+			{
+				.IsDefault = [this]()				{ return GetProfile()->GetCookConfiguration() == GetDefaultProfile()->GetCookConfiguration(); },
+				.SetToDefault = [this]()			{ return GetProfile()->SetCookConfiguration( GetDefaultProfile()->GetCookConfiguration() ); },
+				.IsVisible = IsVisible,
+			},
+			SNew(SCustomLaunchLexToStringCombo<EBuildConfiguration>)
+			.OnSelectionChanged_Lambda( [this](EBuildConfiguration Value) { GetProfile()->SetCookConfiguration(Value); } )
+			.SelectedItem_Lambda( [this]() { return GetProfile()->GetCookConfiguration(); } )
+			.Items(TArray<EBuildConfiguration>{ EBuildConfiguration::Debug, EBuildConfiguration::DebugGame, EBuildConfiguration::Development, EBuildConfiguration::Test, EBuildConfiguration::Shipping })
+		)
+		.AddDirectoryString( LOCTEXT("PackageDirLabel", "Staging Directory Override"),
+			{
+				.GetValue = [this]()				{ return GetProfile()->GetPackageDirectory(); },
+				.SetValue = [this](FString Value)	{ GetProfile()->SetPackageDirectory(Value); },
+				.GetDefaultValue = [this]()			{ return GetDefaultProfile()->GetPackageDirectory(); },
+				.IsVisible = IsVisible,
+				.IsEnabled = IsVisible,
+			}
+		)
+		.AddBoolean( LOCTEXT("IoStoreLabel", "Use container files for optimized loading (I/O Store)"),
+			{
+				.GetValue = [this]()				{ return GetProfile()->IsUsingIoStore(); },
+				.SetValue = [this](bool bValue)		{ GetProfile()->SetUseIoStore(bValue); },
+				.GetDefaultValue = [this]()			{ return GetDefaultProfile()->IsUsingIoStore(); },
+				.IsVisible = IsVisible,
+			}
+		)
+		.AddFileString( LOCTEXT("IoStoreRefContainerLabel", "Optional I/O Store Reference Chunk Database"),
+			{
+				.GetValue = [this]()				{ return GetProfile()->GetReferenceContainerGlobalFileName(); },
+				.SetValue = [this](FString Value)	{ GetProfile()->SetReferenceContainerGlobalFileName(Value); },
+				.GetDefaultValue = [this]()			{ return GetDefaultProfile()->GetReferenceContainerGlobalFileName(); },
+				.IsVisible = IsVisible,
+				.IsEnabled = [this]()				{ return GetProfile()->IsUsingIoStore(); },
+			},
+			TEXT("global.utoc files|global.utoc")
+		)
+		.AddFileString( LOCTEXT("IoStoreRefCryptoKeys", "Optional I/O Store Reference Chunk Crypto Key"),
+			{
+				.GetValue = [this]()				{ return GetProfile()->GetReferenceContainerCryptoKeysFileName(); },
+				.SetValue = [this](FString Value)	{ GetProfile()->SetReferenceContainerCryptoKeysFileName(Value); },
+				.GetDefaultValue = [this]()			{ return GetDefaultProfile()->GetReferenceContainerCryptoKeysFileName(); },
+				.IsVisible = IsVisible,
+				.IsEnabled = [this]()				{ return GetProfile()->IsUsingIoStore(); },
+			},
+			TEXT("crypto.json files|crypto.json")
+		)
+	;
+}
+
+void FDeprecatedPropertiesLaunchExtensionInstance::MakeCustomExtensionSubmenu( FMenuBuilder& MenuBuilder )
+{
+	MenuBuilder.AddMenuEntry(
+		GetExtension()->GetDisplayName(),
+		FText::GetEmpty(),
+		FSlateIcon(),
+		FUIAction( 
+			FExecuteAction::CreateLambda( [this]() { bEnabled = !bEnabled; } ),
+			FCanExecuteAction(),
+			FIsActionChecked::CreateLambda( [this]() { return bEnabled; } )
+			),
+		NAME_None,
+		EUserInterfaceActionType::Check
+	);
+
+	// putting this here for now because it isn't serialized with the profile. @todo: should have a better home!
+	MenuBuilder.AddMenuEntry(
+		LOCTEXT("BuildUATLabel", "Build UAT"),
+		LOCTEXT("BuildUATToolTip", "Build UAT before launching"),
+		FSlateIcon(),
+		FUIAction(
+			FExecuteAction::CreateLambda( [this]() { GetProfile()->SetBuildUAT(!GetProfile()->IsBuildingUAT()); } ),
+			FCanExecuteAction(),
+			FIsActionChecked::CreateLambda( [this]() { return GetProfile()->IsBuildingUAT(); } )
+		),
+		NAME_None,
+		EUserInterfaceActionType::Check
+	);
+}
+
+
+void FDeprecatedPropertiesLaunchExtensionInstance::GetExtensionsMenuEntry( FExtensionsMenuEntry& MenuEntry ) const
+{
+	MenuEntry.SectionName = FString();
+
+	MenuEntry.Type = FExtensionsMenuEntry::Type_SubMenu;
+	MenuEntry.SubmenuName = TEXT("Deprecated");
+	MenuEntry.SubmenuDisplayName = LOCTEXT("DeprecatedSubmenu", "Deprecated");
+}
+
+
+
+TSharedPtr<ProjectLauncher::FLaunchExtensionInstance> FDeprecatedPropertiesLaunchExtension::CreateInstanceForProfile( ProjectLauncher::FLaunchExtensionInstance::FArgs& InArgs )
+{
+	return MakeShared<FDeprecatedPropertiesLaunchExtensionInstance>(InArgs);
+}
+
+const TCHAR* FDeprecatedPropertiesLaunchExtension::GetInternalName() const
+{
+	return TEXT("DeprecatedProperties");
+}
+
+FText FDeprecatedPropertiesLaunchExtension::GetDisplayName() const
+{
+	return LOCTEXT("ExtensionName", "Deprecated Properties");
+}
+
+
+#undef LOCTEXT_NAMESPACE
diff --git a/Engine/Plugins/Developer/ProjectLauncher/Source/CommonLaunchExtensions/Private/DeprecatedProperties/DeprecatedPropertiesLaunchExtension.h b/Engine/Plugins/Developer/ProjectLauncher/Source/CommonLaunchExtensions/Private/DeprecatedProperties/DeprecatedPropertiesLaunchExtension.h
new file mode 100644
--- /dev/null
+++ b/Engine/Plugins/Developer/ProjectLauncher/Source/CommonLaunchExtensions/Private/DeprecatedProperties/DeprecatedPropertiesLaunchExtension.h
@@ -0,0 +1,33 @@
+// Copyright Epic Games, Inc. All Rights Reserved.
+
+#pragma once
+
+#include "Extension/LaunchExtension.h"
+
+
+class FDeprecatedPropertiesLaunchExtensionInstance : public ProjectLauncher::FLaunchExtensionInstance
+{
+public:
+	FDeprecatedPropertiesLaunchExtensionInstance( FArgs& InArgs );
+	virtual ~FDeprecatedPropertiesLaunchExtensionInstance() = default;
+
+	virtual void CustomizeTree( ProjectLauncher::FLaunchProfileTreeData& ProfileTreeData ) override;
+
+	virtual bool HasCustomExtensionMenu() const { return true; }
+	virtual void MakeCustomExtensionSubmenu( FMenuBuilder& MenuBuilder ) override;
+	virtual void GetExtensionsMenuEntry( FExtensionsMenuEntry& MenuEntry ) const override;
+
+protected:
+	bool bEnabled = false;
+	bool bHasBeenShown = false;
+	bool HasDeprecatedProperties() const;
+};
+
+
+class FDeprecatedPropertiesLaunchExtension : public ProjectLauncher::FLaunchExtension
+{
+public:
+	virtual TSharedPtr<ProjectLauncher::FLaunchExtensionInstance> CreateInstanceForProfile( ProjectLauncher::FLaunchExtensionInstance::FArgs& InArgs ) override;
+	virtual const TCHAR* GetInternalName() const override;
+	virtual FText GetDisplayName() const override;
+};
\ No newline at end of file
diff --git a/Engine/Plugins/Developer/ProjectLauncher/Source/ProjectLauncher/Private/Model/ProjectLauncherModel.cpp b/Engine/Plugins/Developer/ProjectLauncher/Source/ProjectLauncher/Private/Model/ProjectLauncherModel.cpp
--- a/Engine/Plugins/Developer/ProjectLauncher/Source/ProjectLauncher/Private/Model/ProjectLauncherModel.cpp
+++ b/Engine/Plugins/Developer/ProjectLauncher/Source/ProjectLauncher/Private/Model/ProjectLauncherModel.cpp
@@ -58,6 +58,11 @@ namespace ProjectLauncher
 		bHasSetBasicLaunchProfilePlatform = (BasicLaunchProfile->GetCookedPlatforms().Num() > 0);
 
 		AllProfiles = ProfileManager->GetAllProfiles();
+		for (const ILauncherProfilePtr& Profile : AllProfiles)
+		{
+			Profile->SetBuildUAT(false); // this isn't serialized, so set a better default
+		}
+
 		AllProfiles.Add(BasicLaunchProfile);
 		SortProfiles();
 
diff --git a/Engine/Plugins/Developer/ProjectLauncher/Source/ProjectLauncher/Private/Widgets/Profiles/SCustomLaunchProfilesPanel.cpp b/Engine/Plugins/Developer/ProjectLauncher/Source/ProjectLauncher/Private/Widgets/Profiles/SCustomLaunchProfilesPanel.cpp
--- a/Engine/Plugins/Developer/ProjectLauncher/Source/ProjectLauncher/Private/Widgets/Profiles/SCustomLaunchProfilesPanel.cpp
+++ b/Engine/Plugins/Developer/ProjectLauncher/Source/ProjectLauncher/Private/Widgets/Profiles/SCustomLaunchProfilesPanel.cpp
@@ -526,7 +526,7 @@ TSharedRef<SWidget> SCustomLaunchProfilesPanel::CreateAdvancedProfileWarningWidg
 			[
 				SNew(STextBlock)
 				.ColorAndOpacity(FLinearColor::White)
-				.Text(LOCTEXT("AdvancedProfileWarningDetail", "These settings may be multiple cook platforms or cultures, or distribution packages, DLC or patch builds. Please use Legacy Project Launcher to edit the advanced properties in this profile."))
+				.Text(LOCTEXT("AdvancedProfileWarningDetail", "These settings may be multiple cook platforms or unsupported packaging, deployment or launch modes. Please use Legacy Project Launcher to edit the advanced properties in this profile."))
 				.AutoWrapText(true)
 			]
 		]
