Project Launcher profile tree improvements
- property tree headings can now specify a sort priority
- new AddSubHeading to create property trees

#rnx
#Synced-CL 48778233

diff --git a/Engine/Plugins/Developer/ProjectLauncher/Source/ProjectLauncher/Private/ProfileTree/LaunchProfileTreeData.cpp b/Engine/Plugins/Developer/ProjectLauncher/Source/ProjectLauncher/Private/ProfileTree/LaunchProfileTreeData.cpp
--- a/Engine/Plugins/Developer/ProjectLauncher/Source/ProjectLauncher/Private/ProfileTree/LaunchProfileTreeData.cpp
+++ b/Engine/Plugins/Developer/ProjectLauncher/Source/ProjectLauncher/Private/ProfileTree/LaunchProfileTreeData.cpp
@@ -139,7 +139,7 @@ namespace ProjectLauncher
 	}
 
 
-	FLaunchProfileTreeNode& FLaunchProfileTreeData::AddHeading( const TCHAR* InName, FText InDisplayName )
+	FLaunchProfileTreeNode& FLaunchProfileTreeData::AddHeading( const TCHAR* InName, FText InDisplayName, int32 InSortOrder )
 	{
 		if (HeadingNodes.Contains(InName))
 		{
@@ -148,10 +148,17 @@ namespace ProjectLauncher
 
 		FLaunchProfileTreeNodePtr TreeNode = MakeShared<FLaunchProfileTreeNode>(this);
 		TreeNode->Name = InDisplayName;
+		TreeNode->SortOrder = InSortOrder;
+		TreeNode->Depth = 0;
 		Nodes.Add(TreeNode);
 
 		HeadingNodes.Add(InName, TreeNode);
 
+		Nodes.StableSort( []( const FLaunchProfileTreeNodePtr& A, const FLaunchProfileTreeNodePtr& B)
+		{
+			return A->SortOrder < B->SortOrder;
+		});
+
 		return *TreeNode.Get();
 	}
 
@@ -197,10 +204,46 @@ namespace ProjectLauncher
 		TreeData->TreeBuilder->OnPropertyChanged();
 	}
 
+
+
+
+	FLaunchProfileTreeNode& FLaunchProfileTreeNode::AddSubHeading( const TCHAR* InName, FText InDisplayName, int32 InSortOrder )
+	{
+		if (SubHeadingNodes.Contains(InName))
+		{
+			return GetSubHeading(InName);
+		}
+
+		FLaunchProfileTreeNodePtr TreeNode = MakeShared<FLaunchProfileTreeNode>(TreeData);
+		TreeNode->Name = InDisplayName;
+		TreeNode->SortOrder = InSortOrder;
+		TreeNode->Depth = Depth+1;
+		Children.Add(TreeNode);
+
+		SubHeadingNodes.Add(InName, TreeNode);
+
+		Children.StableSort( []( const FLaunchProfileTreeNodePtr& A, const FLaunchProfileTreeNodePtr& B)
+		{
+			return A->SortOrder < B->SortOrder;
+		});
+
+		return *TreeNode.Get();
+	}
+
+
+	FLaunchProfileTreeNode& FLaunchProfileTreeNode::GetSubHeading(const TCHAR* InName) const
+	{
+		FLaunchProfileTreeNodePtr TreeNode = SubHeadingNodes.FindChecked(InName);
+		return *TreeNode.Get();
+	}
+
+
 	FLaunchProfileTreeNode& FLaunchProfileTreeNode::AddWidget( FText InName, FCallbacks&& InWidgetCallbacks, TSharedRef<SWidget> InValueWidget )
 	{
 		FLaunchProfileTreeNodePtr TreeNode = MakeShared<FLaunchProfileTreeNode>(TreeData);
 		TreeNode->Name = InName;
+		TreeNode->SortOrder = 0;
+		TreeNode->Depth = Depth+1;
 		TreeNode->Widget = InValueWidget;
 		TreeNode->Callbacks = InWidgetCallbacks;
 		Children.Add(TreeNode);
diff --git a/Engine/Plugins/Developer/ProjectLauncher/Source/ProjectLauncher/Private/Widgets/Profiles/SCustomLaunchCustomProfileEditor.cpp b/Engine/Plugins/Developer/ProjectLauncher/Source/ProjectLauncher/Private/Widgets/Profiles/SCustomLaunchCustomProfileEditor.cpp
--- a/Engine/Plugins/Developer/ProjectLauncher/Source/ProjectLauncher/Private/Widgets/Profiles/SCustomLaunchCustomProfileEditor.cpp
+++ b/Engine/Plugins/Developer/ProjectLauncher/Source/ProjectLauncher/Private/Widgets/Profiles/SCustomLaunchCustomProfileEditor.cpp
@@ -109,7 +109,7 @@ public:
 			[
 				SNew(SBorder)
 				.BorderImage(FAppStyle::GetBrush("Brushes.Header"))				
-				.Padding(0)
+				.Padding(8*TreeNode->Depth,0,0,0)
 				[
 					SNew(SBox)
 					.MinDesiredHeight(26.0f)
@@ -179,7 +179,7 @@ public:
 				.Padding(1)
 				[
 					SNew(SBox)
-					.Padding( FMargin(24, 1, 24, 1))
+					.Padding( FMargin(24*TreeNode->Depth, 1, 24, 1))
 					.VAlign(VAlign_Center)
 					[
 						TreeNode->Widget.ToSharedRef()
@@ -319,7 +319,7 @@ public:
 			[
 				MakeSplitterSlot(
 					SNew(SBox)
-					.Padding( FMargin(24, 1, 4, 0))
+					.Padding( FMargin(8 + (16*TreeNode->Depth), 1, 4, 0))
 					.VAlign(VAlign_Top)
 					[
 						SAssignNew(PropertyNameBox, SHorizontalBox)
diff --git a/Engine/Plugins/Developer/ProjectLauncher/Source/ProjectLauncher/Public/ProfileTree/LaunchProfileTreeData.h b/Engine/Plugins/Developer/ProjectLauncher/Source/ProjectLauncher/Public/ProfileTree/LaunchProfileTreeData.h
--- a/Engine/Plugins/Developer/ProjectLauncher/Source/ProjectLauncher/Public/ProfileTree/LaunchProfileTreeData.h
+++ b/Engine/Plugins/Developer/ProjectLauncher/Source/ProjectLauncher/Public/ProfileTree/LaunchProfileTreeData.h
@@ -145,11 +145,20 @@ namespace ProjectLauncher
 		UE_API FLaunchProfileTreeNode& AddInteger( FText InName, TTypeCallbacks<int32>&& TypeCallbacks, FInt32Interval Range = FInt32Interval(), bool bShowSpinner = false );
 		UE_API FLaunchProfileTreeNode& AddFloat( FText InName, TTypeCallbacks<float>&& TypeCallbacks, FFloatInterval Range = FFloatInterval(), bool bShowSpinner = false );
 
+		UE_API FLaunchProfileTreeNode& AddSubHeading(const TCHAR* InName, FText InDisplayName, int32 SortOrder = 0);
+		UE_API FLaunchProfileTreeNode& GetSubHeading(const TCHAR* InName) const;
+
+
+
 		/* Node details */
-		FText Name;                                 // label for the node
-		TSharedPtr<SWidget> Widget;                 // UI widget
-		FCallbacks Callbacks;                       // generic callbacks
-		TArray<FLaunchProfileTreeNodePtr> Children; // child nodes
+		FText Name;													// label for the node
+		int32 SortOrder = 0;										// sort order for headings
+		int32 Depth = 0;											// tree depth
+		TSharedPtr<SWidget> Widget;									// UI widget
+		FCallbacks Callbacks;										// generic callbacks
+		TArray<FLaunchProfileTreeNodePtr> Children;					// child nodes
+		TMap<FString,FLaunchProfileTreeNodePtr> SubHeadingNodes;    // named subheadings
+
 		const FLaunchProfileTreeData* GetTreeData() const { return TreeData; }
 
 	protected:
@@ -171,7 +180,7 @@ namespace ProjectLauncher
 		UE_API ~FLaunchProfileTreeData();
 
 		/* Add or retrieve top-level headings. */
-		UE_API FLaunchProfileTreeNode& AddHeading(const TCHAR* InName, FText InDisplayName);
+		UE_API FLaunchProfileTreeNode& AddHeading(const TCHAR* InName, FText InDisplayName, int32 SortOrder = 0);
 		UE_API FLaunchProfileTreeNode& GetHeading(const TCHAR* InName) const;
 
 		/* Invoke extension instances to build their UI. */
