Cleaned up Project Launcher extensions menu + easier customisability

#rnx

#Synced-CL 48778233

diff --git a/Engine/Plugins/Developer/ProjectLauncher/Source/CommonLaunchExtensions/Private/DeprecatedProperties/DeprecatedPropertiesLaunchExtension.cpp b/Engine/Plugins/Developer/ProjectLauncher/Source/CommonLaunchExtensions/Private/DeprecatedProperties/DeprecatedPropertiesLaunchExtension.cpp
--- a/Engine/Plugins/Developer/ProjectLauncher/Source/CommonLaunchExtensions/Private/DeprecatedProperties/DeprecatedPropertiesLaunchExtension.cpp
+++ b/Engine/Plugins/Developer/ProjectLauncher/Source/CommonLaunchExtensions/Private/DeprecatedProperties/DeprecatedPropertiesLaunchExtension.cpp
@@ -183,11 +183,7 @@ void FDeprecatedPropertiesLaunchExtensionInstance::MakeCustomExtensionSubmenu( F
 
 void FDeprecatedPropertiesLaunchExtensionInstance::GetExtensionsMenuEntry( FExtensionsMenuEntry& MenuEntry ) const
 {
-	MenuEntry.SectionName = FString();
-
-	MenuEntry.Type = FExtensionsMenuEntry::Type_SubMenu;
-	MenuEntry.SubmenuName = TEXT("Deprecated");
-	MenuEntry.SubmenuDisplayName = LOCTEXT("DeprecatedSubmenu", "Deprecated");
+	MenuEntry = FExtensionsMenuEntry::Deprecated;	
 }
 
 
diff --git a/Engine/Plugins/Developer/ProjectLauncher/Source/CommonLaunchExtensions/Private/ProfileWizard/ProfileWizardLaunchExtension.cpp b/Engine/Plugins/Developer/ProjectLauncher/Source/CommonLaunchExtensions/Private/ProfileWizard/ProfileWizardLaunchExtension.cpp
--- a/Engine/Plugins/Developer/ProjectLauncher/Source/CommonLaunchExtensions/Private/ProfileWizard/ProfileWizardLaunchExtension.cpp
+++ b/Engine/Plugins/Developer/ProjectLauncher/Source/CommonLaunchExtensions/Private/ProfileWizard/ProfileWizardLaunchExtension.cpp
@@ -38,11 +38,7 @@ void FProfileWizardLaunchExtensionInstance::MakeCustomExtensionSubmenu( FMenuBui
 
 void FProfileWizardLaunchExtensionInstance::GetExtensionsMenuEntry( FExtensionsMenuEntry& MenuEntry ) const
 {
-	MenuEntry.SectionName = FString();
-
-	MenuEntry.Type = FExtensionsMenuEntry::Type_SubMenu;
-	MenuEntry.SubmenuName = TEXT("Deprecated");
-	MenuEntry.SubmenuDisplayName = LOCTEXT("DeprecatedSubmenu", "Deprecated");
+	MenuEntry = FExtensionsMenuEntry::Deprecated;	
 }
 
 
diff --git a/Engine/Plugins/Developer/ProjectLauncher/Source/ProjectLauncher/Private/Extension/LaunchExtension.cpp b/Engine/Plugins/Developer/ProjectLauncher/Source/ProjectLauncher/Private/Extension/LaunchExtension.cpp
--- a/Engine/Plugins/Developer/ProjectLauncher/Source/ProjectLauncher/Private/Extension/LaunchExtension.cpp
+++ b/Engine/Plugins/Developer/ProjectLauncher/Source/ProjectLauncher/Private/Extension/LaunchExtension.cpp
@@ -206,6 +206,49 @@ namespace ProjectLauncher
 
 
 
+	// common menu locations
+	namespace
+	{
+		static const TCHAR* DefaultSectionName = TEXT("Common");
+		static const FText DefaultSectionDisplayName = LOCTEXT("CommonSectionName","General Extensions");
+	}
+
+	const FLaunchExtensionInstance::FExtensionsMenuEntry FLaunchExtensionInstance::FExtensionsMenuEntry::None =
+	{
+		.Type = FLaunchExtensionInstance::FExtensionsMenuEntry::Type_None,
+	};
+
+	const FLaunchExtensionInstance::FExtensionsMenuEntry FLaunchExtensionInstance::FExtensionsMenuEntry::Default =
+	{
+		.SectionName = DefaultSectionName,
+		.SectionDisplayName = DefaultSectionDisplayName,
+		.Type = FLaunchExtensionInstance::FExtensionsMenuEntry::Type_Section,
+	};
+
+	const FLaunchExtensionInstance::FExtensionsMenuEntry FLaunchExtensionInstance::FExtensionsMenuEntry::OwnSection =
+	{
+		.Type = FLaunchExtensionInstance::FExtensionsMenuEntry::Type_Section,
+	};
+
+	const FLaunchExtensionInstance::FExtensionsMenuEntry FLaunchExtensionInstance::FExtensionsMenuEntry::Deprecated =
+	{
+		.SectionName = DefaultSectionName,
+		.SubmenuName = TEXT("Deprecated"),
+		.SectionDisplayName = DefaultSectionDisplayName,
+		.SubmenuDisplayName = LOCTEXT("DeprecatedSubmenu", "Deprecated"),
+		.Type = FLaunchExtensionInstance::FExtensionsMenuEntry::Type_SubMenu,
+	};
+
+	const FLaunchExtensionInstance::FExtensionsMenuEntry FLaunchExtensionInstance::FExtensionsMenuEntry::AutomatedTests =
+	{
+		.SectionName = DefaultSectionName,
+		.SubmenuName = TEXT("AutomatedTests"),
+		.SectionDisplayName = DefaultSectionDisplayName,
+		.SubmenuDisplayName = LOCTEXT("AutotestMenuSubMenu", "Automated Tests"),
+		.Type = FLaunchExtensionInstance::FExtensionsMenuEntry::Type_SubMenu,
+	};
+
+
 
 
 
@@ -682,14 +725,7 @@ namespace ProjectLauncher
 
 	void FLaunchExtensionInstance::GetExtensionsMenuEntry( FExtensionsMenuEntry& MenuEntry ) const
 	{
-		// put the extension in its own section by default
-		MenuEntry.Type = FExtensionsMenuEntry::Type_Section;
-		MenuEntry.SectionName = Extension->GetInternalName();
-		MenuEntry.SectionDisplayName = Extension->GetDisplayName();
-		
-		// fill in the default submenu details anyway - an extension can then just change MenuEntry.Type = Type_SubMenu to appear in a named submenu
-		MenuEntry.SubmenuName = Extension->GetInternalName();
-		MenuEntry.SubmenuDisplayName = Extension->GetDisplayName();
+		MenuEntry = FExtensionsMenuEntry::Default;
 	}
 
 
@@ -748,14 +784,7 @@ namespace ProjectLauncher
 
 	void FAutomatedTestLaunchExtensionInstance::GetExtensionsMenuEntry( FExtensionsMenuEntry& MenuEntry ) const
 	{
-		// put all automated tests in shared submenu by default
-		MenuEntry.Type = FExtensionsMenuEntry::Type_SubMenu;
-		MenuEntry.SubmenuName = TEXT("AutomatedTestsSection");
-		MenuEntry.SubmenuDisplayName = LOCTEXT("AutotestMenuSection", "Automated Tests");
-
-		// put the submenu into the 'common' category
-		MenuEntry.SectionName = FString();
-
+		MenuEntry = FExtensionsMenuEntry::AutomatedTests;
 	}
 
 
diff --git a/Engine/Plugins/Developer/ProjectLauncher/Source/ProjectLauncher/Private/Widgets/Profiles/SCustomLaunchCustomProfileEditor.cpp b/Engine/Plugins/Developer/ProjectLauncher/Source/ProjectLauncher/Private/Widgets/Profiles/SCustomLaunchCustomProfileEditor.cpp
--- a/Engine/Plugins/Developer/ProjectLauncher/Source/ProjectLauncher/Private/Widgets/Profiles/SCustomLaunchCustomProfileEditor.cpp
+++ b/Engine/Plugins/Developer/ProjectLauncher/Source/ProjectLauncher/Private/Widgets/Profiles/SCustomLaunchCustomProfileEditor.cpp
@@ -576,6 +576,7 @@ TSharedRef<SWidget> SCustomLaunchCustomProfileEditor::MakeExtensionsMenu()
 
 	struct FExtensionMenuSubMenu
 	{
+		FText MenuEntryName;
 		FText DisplayName;
 		TArray<TSharedPtr<ProjectLauncher::FLaunchExtensionInstance>> ExtensionInstances;
 	};
@@ -599,11 +600,18 @@ TSharedRef<SWidget> SCustomLaunchCustomProfileEditor::MakeExtensionsMenu()
 		}
 
 		// get section
-		FExtensionMenuSection* Section = Sections.Find(ExtensionMenuEntry.SectionName);
+		FString SectionName = ExtensionMenuEntry.SectionName;
+		FText SectionDisplayName = ExtensionMenuEntry.SectionDisplayName;
+		if (SectionName.IsEmpty())
+		{
+			SectionName = ExtensionInstance->GetExtension()->GetInternalName();
+			SectionDisplayName = ExtensionInstance->GetExtension()->GetDisplayName();
+		}
+		FExtensionMenuSection* Section = Sections.Find(SectionName);
 		if (Section == nullptr)
 		{
-			Section = &Sections.Add(ExtensionMenuEntry.SectionName);
-			Section->DisplayName = ExtensionMenuEntry.SectionName.IsEmpty() ? FText::GetEmpty() : ExtensionMenuEntry.SectionDisplayName;
+			Section = &Sections.Add(SectionName);
+			Section->DisplayName = SectionDisplayName;
 		}
 
 		if (ExtensionMenuEntry.Type == ProjectLauncher::FLaunchExtensionInstance::FExtensionsMenuEntry::Type_SubMenu)
@@ -618,13 +626,13 @@ TSharedRef<SWidget> SCustomLaunchCustomProfileEditor::MakeExtensionsMenu()
 
 			(*SubMenu)->ExtensionInstances.Add(ExtensionInstance);
 		}
-		else if (ExtensionMenuEntry.Type == ProjectLauncher::FLaunchExtensionInstance::FExtensionsMenuEntry::Type_Section)
+		else
 		{
 			Section->ExtensionInstances.Add(ExtensionInstance);
 		}
 	}
 
-	// sort sections alphabetically : this should leave the 'empty' section at the top
+	// sort sections alphabetically
 	Sections.ValueSort([](const FExtensionMenuSection& A, const FExtensionMenuSection& B)
 	{
 		return A.DisplayName.CompareToCaseIgnored(B.DisplayName) < 0; 
@@ -641,11 +649,8 @@ TSharedRef<SWidget> SCustomLaunchCustomProfileEditor::MakeExtensionsMenu()
 	
 	for (const auto& SectionItr : Sections)
 	{
-		// start the section if necessary
-		if (!SectionItr.Key.IsEmpty())
-		{
-			MenuBuilder.BeginSection(NAME_None, SectionItr.Value.DisplayName);
-		}
+		// start the section
+		MenuBuilder.BeginSection(NAME_None, SectionItr.Value.DisplayName);
 
 		// add all extensions
 		for (TSharedPtr<ProjectLauncher::FLaunchExtensionInstance> ExtensionInstance : SectionItr.Value.ExtensionInstances)
@@ -658,12 +663,11 @@ TSharedRef<SWidget> SCustomLaunchCustomProfileEditor::MakeExtensionsMenu()
 		{
 			auto MakeExtensionSubMenu = [](FMenuBuilder& MenuBuilder, TSharedPtr<FExtensionMenuSubMenu> SubMenu)
 			{
-				MenuBuilder.BeginSection(NAME_None, SubMenu->DisplayName);
+				MenuBuilder.SetSearchable(bSearchable);
 				for (TSharedPtr<ProjectLauncher::FLaunchExtensionInstance> ExtensionInstance : SubMenu->ExtensionInstances)
 				{
 					ExtensionInstance->MakeCustomExtensionSubmenu(MenuBuilder);
 				}
-				MenuBuilder.EndSection();
 			};
 
 			const bool bInOpenSubMenuOnClick = false;
@@ -677,11 +681,8 @@ TSharedRef<SWidget> SCustomLaunchCustomProfileEditor::MakeExtensionsMenu()
 			);
 		}
 
-		// end the section, if necessary
-		if (!SectionItr.Key.IsEmpty())
-		{
-			MenuBuilder.EndSection();
-		}
+		// end the section
+		MenuBuilder.EndSection();
 	}
 
 	return MenuBuilder.MakeWidget();
diff --git a/Engine/Plugins/Developer/ProjectLauncher/Source/ProjectLauncher/Public/Extension/LaunchExtension.h b/Engine/Plugins/Developer/ProjectLauncher/Source/ProjectLauncher/Public/Extension/LaunchExtension.h
--- a/Engine/Plugins/Developer/ProjectLauncher/Source/ProjectLauncher/Public/Extension/LaunchExtension.h
+++ b/Engine/Plugins/Developer/ProjectLauncher/Source/ProjectLauncher/Public/Extension/LaunchExtension.h
@@ -108,8 +108,15 @@ namespace ProjectLauncher
 		 */
 		struct FExtensionsMenuEntry
 		{
+			// common menu locations
+			static UE_API const FExtensionsMenuEntry None;
+			static UE_API const FExtensionsMenuEntry Default;
+			static UE_API const FExtensionsMenuEntry OwnSection;
+			static UE_API const FExtensionsMenuEntry Deprecated;
+			static UE_API const FExtensionsMenuEntry AutomatedTests;
+
 			// internal names
-			FString SectionName; // if this is empty, item will be placed in the common section & SectionDisplayName will be ignored
+			FString SectionName; // if this is empty, the extension will be placed in its own section
 			FString SubmenuName;
 
 			// friendly names
@@ -119,11 +126,11 @@ namespace ProjectLauncher
 			// type of menu entry
 			enum EType
 			{
-				Type_None,		// does not appear in the extensions menu
-				Type_Section,	// appears in the extensions menu in the SectionName section.
-				Type_SubMenu,	// appears in the extensions menu in a submenu of the SectionName section
+				Type_None,			// does not appear in the extensions menu
+				Type_Section,		// appears in the extensions menu in the common section.
+				Type_SubMenu,		// appears in the extensions menu in a submenu of the SectionName section
 			};
-			EType Type;
+			EType Type = Type_None;
 		};
 		UE_API virtual void GetExtensionsMenuEntry( FExtensionsMenuEntry& MenuEntry ) const;
 
