Bug fix: Project Launcher extension instances are re-initialized on startup

#rnx
#Synced-CL 48778233

diff --git a/Engine/Plugins/Developer/ProjectLauncher/Source/ProjectLauncher/Private/Extension/LaunchExtension.cpp b/Engine/Plugins/Developer/ProjectLauncher/Source/ProjectLauncher/Private/Extension/LaunchExtension.cpp
--- a/Engine/Plugins/Developer/ProjectLauncher/Source/ProjectLauncher/Private/Extension/LaunchExtension.cpp
+++ b/Engine/Plugins/Developer/ProjectLauncher/Source/ProjectLauncher/Private/Extension/LaunchExtension.cpp
@@ -191,6 +191,7 @@ namespace ProjectLauncher
 			TSharedPtr<FLaunchExtensionInstance> ExtensionInstance = Extension->CreateInstanceForProfile(Args);
 			if (ExtensionInstance.IsValid())
 			{
+				ExtensionInstance->InternalInitialize();
 				Result.Add(ExtensionInstance);
 			}
 		}
@@ -740,6 +741,15 @@ namespace ProjectLauncher
 	{
 	}
 
+	void FAutomatedTestLaunchExtensionInstance::InternalInitialize()
+	{
+		ILauncherProfileAutomatedTestPtr AutomatedTest = GetTest();
+		if (AutomatedTest.IsValid())
+		{
+			OnTestAdded(AutomatedTest.ToSharedRef());
+		}
+	}
+
 	void FAutomatedTestLaunchExtensionInstance::MakeCustomExtensionSubmenu( FMenuBuilder& MenuBuilder )
 	{
 		auto ToggleTestEnabled = [this]()
diff --git a/Engine/Plugins/Developer/ProjectLauncher/Source/ProjectLauncher/Public/Extension/LaunchExtension.h b/Engine/Plugins/Developer/ProjectLauncher/Source/ProjectLauncher/Public/Extension/LaunchExtension.h
--- a/Engine/Plugins/Developer/ProjectLauncher/Source/ProjectLauncher/Public/Extension/LaunchExtension.h
+++ b/Engine/Plugins/Developer/ProjectLauncher/Source/ProjectLauncher/Public/Extension/LaunchExtension.h
@@ -160,6 +160,11 @@ namespace ProjectLauncher
 		 */
 		inline FPropertyChangedDelegate& GetPropertyChangedDelegate() { return PropertyChangedDelegate; }
 
+		/**
+		 * Internal intitialization function, called immediately after construction - don't call directly
+		 */
+		UE_API virtual void InternalInitialize() {};
+
 	protected:
 
 		/**
@@ -433,6 +438,11 @@ namespace ProjectLauncher
 		 * Internal name for the automated test. Must be globally unique, so recommended to include the extension name
 		 */
 		virtual const FString GetTestInternalName() const = 0;
+
+		/**
+		 * Internal intitialization function, called immediately after construction - don't call directly
+		 */
+		UE_API virtual void InternalInitialize() override;
 	};
 
 
