Add $(LocalHosts) to the Globals launch extension and made $(LocalHost) return the first actual ip address, not loopback

#rnx
#Synced-CL 48778233

diff --git a/Engine/Plugins/Developer/ProjectLauncher/Source/CommonLaunchExtensions/Private/Globals/GlobalsLaunchExtension.cpp b/Engine/Plugins/Developer/ProjectLauncher/Source/CommonLaunchExtensions/Private/Globals/GlobalsLaunchExtension.cpp
--- a/Engine/Plugins/Developer/ProjectLauncher/Source/CommonLaunchExtensions/Private/Globals/GlobalsLaunchExtension.cpp
+++ b/Engine/Plugins/Developer/ProjectLauncher/Source/CommonLaunchExtensions/Private/Globals/GlobalsLaunchExtension.cpp
@@ -9,6 +9,7 @@
 
 
 const TCHAR* FGlobalsLaunchExtensionInstance::LocalHostVariable     = TEXT("$(LocalHost)");
+const TCHAR* FGlobalsLaunchExtensionInstance::LocalHostsVariable    = TEXT("$(LocalHosts)");
 const TCHAR* FGlobalsLaunchExtensionInstance::ProjectNameVariable   = TEXT("$(ProjectName)");
 const TCHAR* FGlobalsLaunchExtensionInstance::ProjectPathVariable   = TEXT("$(ProjectPath)");
 const TCHAR* FGlobalsLaunchExtensionInstance::TargetNameVariable    = TEXT("$(TargetName)");
@@ -18,6 +19,7 @@ const TCHAR* FGlobalsLaunchExtensionInstance::ConfigurationVariable = TEXT("$(Co
 bool FGlobalsLaunchExtensionInstance::GetExtensionVariables( TArray<FString>& OutVariables ) const
 {
 	OutVariables.Add(LocalHostVariable);
+	OutVariables.Add(LocalHostsVariable);
 	OutVariables.Add(ProjectNameVariable);
 	OutVariables.Add(ProjectPathVariable);
 	OutVariables.Add(TargetNameVariable);
@@ -28,18 +30,36 @@ bool FGlobalsLaunchExtensionInstance::GetExtensionVariables( TArray<FString>& Ou
 
 bool FGlobalsLaunchExtensionInstance::GetExtensionVariableValue( const FString& InVariable, FString& OutValue ) const
 {
-	if (InVariable == LocalHostVariable)
+	if (InVariable == LocalHostVariable || InVariable == LocalHostsVariable)
 	{
-		// @todo: what if we have multiple host addresses etc etc. ?
-		if (ISocketSubsystem* SocketSubsystem = ISocketSubsystem::Get(PLATFORM_SOCKETSUBSYSTEM))
+		bool bSingle = (InVariable == LocalHostVariable);
+
+		TArray<TSharedPtr<FInternetAddr>> AdapterAddresses;
+		ISocketSubsystem::Get(PLATFORM_SOCKETSUBSYSTEM)->GetLocalAdapterAddresses(AdapterAddresses);
+		if (AdapterAddresses.Num() > 0)
 		{
-			bool bCanBindAll;
-			TSharedPtr<FInternetAddr> LocalHostAddr = SocketSubsystem->GetLocalHostAddr(*GLog, bCanBindAll);
+			TArray<FString> LocalHosts;
+			for (TSharedPtr<FInternetAddr>& LocalHostAddr : AdapterAddresses)
+			{
+				if (LocalHostAddr.IsValid())
+				{
+					constexpr bool bAppendPort = false;
+					FString AddressString = LocalHostAddr->ToString(bAppendPort);
+					if (bSingle) // @todo: any better heuristic than "first" ?
+					{
+						OutValue = AddressString;
+						return true;
+					}
+					else
+					{
+						LocalHosts.Add(AddressString);
+					}
+				}
+			}
 
-			if (LocalHostAddr->IsValid())
+			if (LocalHosts.Num() > 0)
 			{
-				constexpr bool bAppendPort = false;
-				OutValue = LocalHostAddr->ToString(bAppendPort);
+				OutValue = FString::Join(LocalHosts, TEXT("+"));
 				return true;
 			}
 		}
diff --git a/Engine/Plugins/Developer/ProjectLauncher/Source/CommonLaunchExtensions/Private/Globals/GlobalsLaunchExtension.h b/Engine/Plugins/Developer/ProjectLauncher/Source/CommonLaunchExtensions/Private/Globals/GlobalsLaunchExtension.h
--- a/Engine/Plugins/Developer/ProjectLauncher/Source/CommonLaunchExtensions/Private/Globals/GlobalsLaunchExtension.h
+++ b/Engine/Plugins/Developer/ProjectLauncher/Source/CommonLaunchExtensions/Private/Globals/GlobalsLaunchExtension.h
@@ -16,6 +16,7 @@ public:
 private:
 
 	static const TCHAR* LocalHostVariable;
+	static const TCHAR* LocalHostsVariable;
 	static const TCHAR* ProjectNameVariable;
 	static const TCHAR* ProjectPathVariable;
 	static const TCHAR* TargetNameVariable;
