fix typo in log file export

#rnx

#Synced-CL 48803560

diff --git a/Engine/Plugins/Developer/ProjectLauncher/Source/ProjectLauncher/Private/Widgets/Output/SCustomLaunchOutputLog.cpp b/Engine/Plugins/Developer/ProjectLauncher/Source/ProjectLauncher/Private/Widgets/Output/SCustomLaunchOutputLog.cpp
--- a/Engine/Plugins/Developer/ProjectLauncher/Source/ProjectLauncher/Private/Widgets/Output/SCustomLaunchOutputLog.cpp
+++ b/Engine/Plugins/Developer/ProjectLauncher/Source/ProjectLauncher/Private/Widgets/Output/SCustomLaunchOutputLog.cpp
@@ -458,7 +458,7 @@ void SCustomLaunchOutputLog::SaveLog()
 			FSlateApplication::Get().FindBestParentWindowHandleForDialogs(nullptr),
 			LOCTEXT("SaveLogDialogTitle", "Save Log As...").ToString(),
 			*LastSavePath,
-			TEXT("BuidCookRun.log"),
+			TEXT("BuildCookRun.log"),
 			TEXT("Log Files (*.log)|*.log"),
 			EFileDialogFlags::None,
 			FileNames))
