Split out Project Launcher extension command line manipulation into util functions

#rnx
#Synced-CL 48803560

diff --git a/Engine/Plugins/Developer/ProjectLauncher/Source/CommonLaunchExtensions/Private/Insights/InsightsLaunchExtension.cpp b/Engine/Plugins/Developer/ProjectLauncher/Source/CommonLaunchExtensions/Private/Insights/InsightsLaunchExtension.cpp
--- a/Engine/Plugins/Developer/ProjectLauncher/Source/CommonLaunchExtensions/Private/Insights/InsightsLaunchExtension.cpp
+++ b/Engine/Plugins/Developer/ProjectLauncher/Source/CommonLaunchExtensions/Private/Insights/InsightsLaunchExtension.cpp
@@ -1,6 +1,7 @@
 // Copyright Epic Games, Inc. All Rights Reserved.
 
 #include "Insights/InsightsLaunchExtension.h"
+#include "Utils/LauncherCmdLineUtils.h"
 #include "Trace/Trace.h"
 
 
@@ -46,7 +47,7 @@ void FInsightsLaunchExtensionInstance::CacheTraceChannels()
 {
 	TraceChannels.Reset();
 
-	FString TraceParamValue = GetParameterValue(TraceParam);
+	FString TraceParamValue = ProjectLauncher::CmdLineUtils::GetParameterValue(GetCommandLine(), TraceParam);
 	TraceParamValue.ParseIntoArray(TraceChannels, TEXT(",") );
 
 	TraceChannels.Sort();
@@ -63,15 +64,17 @@ void FInsightsLaunchExtensionInstance::ToggleTraceChannel( const FString& InChan
 		TraceChannels.Add(InChannel);
 	}
 
+	FString CommandLine = GetCommandLine();
 	if (TraceChannels.IsEmpty())
 	{
-		RemoveParameter(TraceParam);
+		ProjectLauncher::CmdLineUtils::RemoveParameter(CommandLine, TraceParam);
 	}
 	else
 	{
 		FString TraceParamValue = FString::Join( TraceChannels, TEXT(",") );
-		UpdateParameterValue( TraceParam, TraceParamValue );
+		ProjectLauncher::CmdLineUtils::UpdateParameterValue( CommandLine, TraceParam, TraceParamValue );
 	}
+	SetCommandLine(CommandLine);
 }
 
 bool FInsightsLaunchExtensionInstance::IsTraceChannelEnabled( const FString InChannel ) const
@@ -103,13 +106,22 @@ void FInsightsLaunchExtensionInstance::CustomizeParametersSubmenu( FMenuBuilder&
 		}, &AllTraceChannels );
 		AllTraceChannels.Sort();
 
+		auto SetNone = [this]()
+		{
+			FString CommandLine = GetCommandLine();
+			ProjectLauncher::CmdLineUtils::RemoveParameter(CommandLine, TraceParam);
+			SetCommandLine(CommandLine);
+
+			TraceChannels.Reset();
+		};
+
 		// add menu item to clear the selected channels
 		MenuBuilder.AddMenuEntry(
 			LOCTEXT("ClearTraceChannelsLabel", "None"),
 			FText::GetEmpty(),
 			FSlateIcon(),
 			FUIAction( 
-				FExecuteAction::CreateLambda( [this]() { RemoveParameter(TraceParam); TraceChannels.Reset(); } ),
+				FExecuteAction::CreateLambda(SetNone),
 				FCanExecuteAction(),
 				FIsActionChecked::CreateLambda( [this]() { return (TraceChannels.IsEmpty()); } )
 				),
diff --git a/Engine/Plugins/Developer/ProjectLauncher/Source/ProjectLauncher/Private/Extension/LaunchExtension.cpp b/Engine/Plugins/Developer/ProjectLauncher/Source/ProjectLauncher/Private/Extension/LaunchExtension.cpp
--- a/Engine/Plugins/Developer/ProjectLauncher/Source/ProjectLauncher/Private/Extension/LaunchExtension.cpp
+++ b/Engine/Plugins/Developer/ProjectLauncher/Source/ProjectLauncher/Private/Extension/LaunchExtension.cpp
@@ -2,11 +2,11 @@
 
 #include "Extension/LaunchExtension.h"
 #include "Model/ProjectLauncherModel.h"
+#include "Utils/LauncherCmdLineUtils.h"
 #include "Containers/Array.h"
 #include "Templates/SharedPointer.h"
 #include "HAL/CriticalSection.h"
 #include "Misc/ScopeLock.h"
-#include "Misc/Parse.h"
 #include "Misc/ConfigCacheIni.h"
 #include "Modules/ModuleManager.h"
 #include "ILauncherServicesModule.h"
@@ -318,177 +318,56 @@ namespace ProjectLauncher
 
 	FString FLaunchExtensionInstance::GetParameterValue( const FString& InParameter ) const
 	{
-		FString CommandLine = GetCommandLine();
-
-		// InParameter is -Key=
-		FString ParamValue;
-		if (FParse::Value(*CommandLine, *InParameter, ParamValue, false))
-		{
-			return ParamValue;
-		}
-
-		// InParameter is -Key=Value
-		FString ParamKey;
-		if (InParameter.Split(TEXT("="), &ParamKey, &ParamValue))
-		{
-			ParamKey += TEXT("=");
-			if (FParse::Value(*CommandLine, *ParamKey, ParamValue, false))
-			{
-				return ParamValue;
-			}
-		}
-
-		return FString();
+		return CmdLineUtils::GetParameterValue(GetCommandLine(), InParameter);
 	}
 
 	bool FLaunchExtensionInstance::UpdateParameterValue( const FString& InParameter, const FString& NewValue )
 	{
-		FString ParamKey, ParamValue;
-		if (InParameter.Split(TEXT("="), &ParamKey, &ParamValue))
+		FString CommandLine = GetCommandLine();
+		if (CmdLineUtils::UpdateParameterValue(CommandLine, InParameter, NewValue))
 		{
-			RemoveParameter(InParameter);
-
-			FString NewParameter = ParamKey + "=" + NewValue;
-			AddParameter(NewParameter);
+			SetCommandLine(CommandLine);
 			return true;
 		}
 
 		return false;
-
 	}
 
 
 
 	FString FLaunchExtensionInstance::GetFinalParameter( const FString& InParameter ) const
 	{
-		FString CommandLine = GetCommandLine();
-
-		// get the parameter's key & value
-		FString ParamName, ParamValue;
-		bool bParameterHasValue = InParameter.Split(TEXT("="), &ParamName, &ParamValue);
-		if (!bParameterHasValue)
-		{
-			ParamName = InParameter;
-		}
-
-		// the parameter's value may have been added or alterered - return how it is now
-		FString ParamKey = ParamName + TEXT("=");
-		if (FParse::Value(*CommandLine, *ParamKey, ParamValue, false))
-		{
-			return ParamKey + ParamValue;
-		}
-		
-		// the parameter may not have a value or it has been removed - return how it is now
-		FString Param = ParamName;
-		if (Param.RemoveFromStart(TEXT("-")) && FParse::Param(*CommandLine, *Param))
-		{
-			return ParamName;
-		}
-
-		// return the parameter as-is
-		return InParameter;
+		return CmdLineUtils::GetFinalParameter(GetCommandLine(), InParameter);
 	}
 
 	void FLaunchExtensionInstance::AddParameter( const FString& InParameter )
 	{
-		FString CommandLine = GetCommandLine() + TEXT(" ") + InParameter;
+		FString CommandLine = GetCommandLine();
+		CmdLineUtils::AddParameter(CommandLine, InParameter);
 		SetCommandLine(CommandLine);
 	}
 
 	void FLaunchExtensionInstance::RemoveParameter( const FString& InParameter )
 	{
-		if (TryRemoveParameterGroup(InParameter))
-		{
-			return;
-		}
-
-		FString Parameter = GetFinalParameter(InParameter);
 		FString CommandLine = GetCommandLine();
-
-		// first try to remove the parameter and the preceding space
-		FString ParameterWithSpace = TEXT(" ") + Parameter;
-		if (CommandLine.ReplaceInline(*ParameterWithSpace, TEXT(""), ESearchCase::IgnoreCase) > 0)
-		{
-			SetCommandLine(CommandLine);
-			return;
-		}
-
-		// next try to remove the parameter with a trailing space
-		ParameterWithSpace = Parameter + TEXT(" ");
-		if (CommandLine.ReplaceInline(*ParameterWithSpace, TEXT(""), ESearchCase::IgnoreCase) > 0)
-		{
-			SetCommandLine(CommandLine);
-			return;
-		}
-
-		// a little unexpected... just remove the parameter on its own
-		if (CommandLine.ReplaceInline(*Parameter, TEXT(""), ESearchCase::IgnoreCase) > 0)
+		if (CmdLineUtils::RemoveParameter(CommandLine, InParameter))
 		{
 			SetCommandLine(CommandLine);
-			return;
 		}
 	}
 
 
 	bool FLaunchExtensionInstance::IsParameterUsed( const FString& InParameter ) const
 	{
-		FString Parameter = GetFinalParameter(InParameter);
-		return GetCommandLine().Contains(Parameter, ESearchCase::IgnoreCase);
+		return CmdLineUtils::IsParameterUsed(GetCommandLine(), InParameter);
 	}
 
 	void FLaunchExtensionInstance::SetParameterUsed( const FString& InParameter, bool bUsed )
 	{
-		bool bIsUsed = IsParameterUsed(InParameter);
-		if (bIsUsed != bUsed)
-		{
-			if (bIsUsed)
-			{
-				RemoveParameter(InParameter);
-			}
-			else
-			{
-				AddParameter(InParameter);
-			}
-		}
-	}
-
-
-	bool FLaunchExtensionInstance::IsParameterGroup( const FString& InParameter ) const
-	{
-		// see if InParameter contains just a simple -Key=Value or -Param
-		FString ParsedParameter = InParameter;
-		FString ParamKey, ParamValue;
-		if (InParameter.Split(TEXT("="), &ParamKey, &ParamValue))
-		{
-			ParamKey += TEXT("=");
-			if (FParse::Value(*InParameter, *ParamKey, ParamValue, false))
-			{
-				ParsedParameter = ParamKey + ParamValue;
-			}
-		}
-		
-		return (ParsedParameter.TrimStartAndEnd() != InParameter.TrimStartAndEnd());
-	}
-
-	bool FLaunchExtensionInstance::TryRemoveParameterGroup( const FString& InParameter )
-	{
-		if (!IsParameterGroup(InParameter))
-		{
-			return false;
-		}
-
-		// it may contain a group - handle them separately
-		const TCHAR* Ptr = InParameter.GetCharArray().GetData();
-		FString SubParameter = FParse::Token(Ptr, false);
-		while (!SubParameter.IsEmpty())
-		{
-			RemoveParameter(SubParameter);
-			SubParameter = FParse::Token(Ptr, false);
-		}
-
-		return true;
+		FString CommandLine = GetCommandLine();
+		CmdLineUtils::SetParameterUsed(CommandLine, InParameter, bUsed);
+		SetCommandLine(CommandLine);
 	}
-		
 
 
 
@@ -496,14 +375,16 @@ namespace ProjectLauncher
 	{
 		auto ToggleParameter = [this]( FString Parameter )
 		{
-			if (IsParameterUsed(Parameter))
+			FString CommandLine = GetCommandLine();
+			if (CmdLineUtils::IsParameterUsed(CommandLine, Parameter))
 			{
-				RemoveParameter(Parameter);
+				CmdLineUtils::RemoveParameter(CommandLine, Parameter);
 			}
 			else
 			{
-				AddParameter(Parameter);
+				CmdLineUtils::AddParameter(CommandLine, Parameter);
 			}
+			SetCommandLine(CommandLine);
 		};
 
 		TArray<FString> Parameters;
@@ -521,7 +402,7 @@ namespace ProjectLauncher
 					FUIAction(
 						FExecuteAction::CreateLambda(ToggleParameter, Parameter),
 						FCanExecuteAction(),
-						FIsActionChecked::CreateLambda( [this, Parameter]() { return IsParameterUsed(Parameter); } )
+						FIsActionChecked::CreateLambda( [this, Parameter]() { return CmdLineUtils::IsParameterUsed(GetCommandLine(), Parameter); } )
 					),
 					NAME_None,
 					EUserInterfaceActionType::Check);
diff --git a/Engine/Plugins/Developer/ProjectLauncher/Source/ProjectLauncher/Private/Utils/LauncherCmdLineUtils.cpp b/Engine/Plugins/Developer/ProjectLauncher/Source/ProjectLauncher/Private/Utils/LauncherCmdLineUtils.cpp
new file mode 100644
--- /dev/null
+++ b/Engine/Plugins/Developer/ProjectLauncher/Source/ProjectLauncher/Private/Utils/LauncherCmdLineUtils.cpp
@@ -0,0 +1,190 @@
+// Copyright Epic Games, Inc. All Rights Reserved.
+
+#include "Utils/LauncherCmdLineUtils.h"
+#include "Misc/Parse.h"
+
+
+namespace ProjectLauncher::CmdLineUtils
+{
+
+
+	bool IsParameterGroup( const FString& InCmdLine, const FString& InParameter )
+	{
+		// see if InParameter contains just a simple -Key=Value or -Param
+		FString ParsedParameter = InParameter;
+		FString ParamKey, ParamValue;
+		if (InParameter.Split(TEXT("="), &ParamKey, &ParamValue))
+		{
+			ParamKey += TEXT("=");
+			if (FParse::Value(*InParameter, *ParamKey, ParamValue, false))
+			{
+				ParsedParameter = ParamKey + ParamValue;
+			}
+		}
+		
+		return (ParsedParameter.TrimStartAndEnd() != InParameter.TrimStartAndEnd());
+	}
+
+
+	bool TryRemoveParameterGroup( FString& InOutCmdLine, const FString& InParameter )
+	{
+		if (!IsParameterGroup(InOutCmdLine, InParameter))
+		{
+			return false;
+		}
+
+		// it may contain a group - handle them separately
+		const TCHAR* Ptr = InParameter.GetCharArray().GetData();
+		FString SubParameter = FParse::Token(Ptr, false);
+		while (!SubParameter.IsEmpty())
+		{
+			RemoveParameter(InOutCmdLine, SubParameter);
+			SubParameter = FParse::Token(Ptr, false);
+		}
+
+		return true;
+	}
+
+
+	FString GetParameterValue( const FString& InCmdLine, const FString& InParameter )
+	{
+		FString CommandLine = InCmdLine;
+
+		// InParameter is -Key=
+		FString ParamValue;
+		if (FParse::Value(*CommandLine, *InParameter, ParamValue, false))
+		{
+			return ParamValue;
+		}
+
+		// InParameter is -Key=Value
+		FString ParamKey;
+		if (InParameter.Split(TEXT("="), &ParamKey, &ParamValue))
+		{
+			ParamKey += TEXT("=");
+			if (FParse::Value(*CommandLine, *ParamKey, ParamValue, false))
+			{
+				return ParamValue;
+			}
+		}
+
+		return FString();
+	}
+
+
+	bool UpdateParameterValue( FString& InOutCmdLine, const FString& InParameter, const FString& NewValue )
+	{
+		FString ParamKey, ParamValue;
+		if (InParameter.Split(TEXT("="), &ParamKey, &ParamValue))
+		{
+			RemoveParameter(InOutCmdLine, InParameter);
+
+			FString NewParameter = ParamKey + "=" + NewValue;
+			AddParameter(InOutCmdLine, NewParameter);
+			return true;
+		}
+
+		return false;
+
+	}
+
+
+	FString GetFinalParameter( const FString& InCmdLine, const FString& InParameter )
+	{
+		FString CommandLine = InCmdLine;
+
+		// get the parameter's key & value
+		FString ParamName, ParamValue;
+		bool bParameterHasValue = InParameter.Split(TEXT("="), &ParamName, &ParamValue);
+		if (!bParameterHasValue)
+		{
+			ParamName = InParameter;
+		}
+
+		// the parameter's value may have been added or alterered - return how it is now
+		FString ParamKey = ParamName + TEXT("=");
+		if (FParse::Value(*CommandLine, *ParamKey, ParamValue, false))
+		{
+			return ParamKey + ParamValue;
+		}
+		
+		// the parameter may not have a value or it has been removed - return how it is now
+		FString Param = ParamName;
+		if (Param.RemoveFromStart(TEXT("-")) && FParse::Param(*CommandLine, *Param))
+		{
+			return ParamName;
+		}
+
+		// return the parameter as-is
+		return InParameter;
+	}
+
+
+	void AddParameter( FString& InOutCmdLine, const FString& InParameter )
+	{
+		FString CommandLine = InOutCmdLine + TEXT(" ") + InParameter;
+		InOutCmdLine = CommandLine;
+	}
+
+
+	bool RemoveParameter( FString& InOutCmdLine, const FString& InParameter )
+	{
+		if (TryRemoveParameterGroup(InOutCmdLine, InParameter))
+		{
+			return true;
+		}
+
+		FString Parameter = GetFinalParameter(InOutCmdLine, InParameter);
+		FString CommandLine = InOutCmdLine;
+
+		// first try to remove the parameter and the preceding space
+		FString ParameterWithSpace = TEXT(" ") + Parameter;
+		if (CommandLine.ReplaceInline(*ParameterWithSpace, TEXT(""), ESearchCase::IgnoreCase) > 0)
+		{
+			InOutCmdLine = CommandLine;
+			return true;
+		}
+
+		// next try to remove the parameter with a trailing space
+		ParameterWithSpace = Parameter + TEXT(" ");
+		if (CommandLine.ReplaceInline(*ParameterWithSpace, TEXT(""), ESearchCase::IgnoreCase) > 0)
+		{
+			InOutCmdLine = CommandLine;
+			return true;
+		}
+
+		// a little unexpected... just remove the parameter on its own
+		if (CommandLine.ReplaceInline(*Parameter, TEXT(""), ESearchCase::IgnoreCase) > 0)
+		{
+			InOutCmdLine = CommandLine;
+			return true;
+		}
+
+		return false;
+	}
+
+
+	bool IsParameterUsed( const FString& InCmdLine, const FString& InParameter )
+	{
+		FString Parameter = GetFinalParameter(InCmdLine, InParameter);
+		return InCmdLine.Contains(Parameter, ESearchCase::IgnoreCase);
+	}
+
+
+	void SetParameterUsed( FString& InOutCmdLine, const FString& InParameter, bool bUsed )
+	{
+		bool bIsUsed = IsParameterUsed(InOutCmdLine, InParameter);
+		if (bIsUsed != bUsed)
+		{
+			if (bIsUsed)
+			{
+				RemoveParameter(InOutCmdLine, InParameter);
+			}
+			else
+			{
+				AddParameter(InOutCmdLine, InParameter);
+			}
+		}
+	}
+}
+
diff --git a/Engine/Plugins/Developer/ProjectLauncher/Source/ProjectLauncher/Public/Extension/LaunchExtension.h b/Engine/Plugins/Developer/ProjectLauncher/Source/ProjectLauncher/Public/Extension/LaunchExtension.h
--- a/Engine/Plugins/Developer/ProjectLauncher/Source/ProjectLauncher/Public/Extension/LaunchExtension.h
+++ b/Engine/Plugins/Developer/ProjectLauncher/Source/ProjectLauncher/Public/Extension/LaunchExtension.h
@@ -173,64 +173,30 @@ namespace ProjectLauncher
 
 	protected:
 
-		/**
-		 * Determine if the given parameter is on the command line
- 		 * 
-		 * @param InParameter the parameter to test, e.g. -key=value or -param
-		 * @returns whether the parameter is being used
-		 */
+		UE_DEPRECATED(5.8, "Use ProjectLauncher::CmdLineUtils version")
 		UE_API bool IsParameterUsed( const FString& InParameter ) const;
 
-		/** 
-		 * Add or remove the given parameter on the command line
-		 * 
-		 * @param InParameter the parameter to add or remove, e.g. -key=value or -param
-		 * @param bUsed whether to add or remove the parameter
-		 */
+		UE_DEPRECATED(5.8, "Use ProjectLauncher::CmdLineUtils version")
 		UE_API void SetParameterUsed( const FString& InParameter, bool bUsed );
 
-		/** 
-		 * Add the given parameter to the command line
-		 * 
-		 * @param InParameter the parameter to add, e.g. -key=value or -param
-		 */
+		UE_DEPRECATED(5.8, "Use ProjectLauncher::CmdLineUtils version")
 		UE_API void AddParameter( const FString& InParameter );
 
-		/** 
-		 * Remove the given parameter from the command line
-		 * 
-		 * @param InParameter the parameter to remove, e.g. -key=value or -param
-		 */
+		UE_DEPRECATED(5.8, "Use ProjectLauncher::CmdLineUtils version")
 		UE_API void RemoveParameter( const FString& InParameter );
 
-		/** 
-		 * Retrieve the current value for the given command line parameter
-		 * 
-		 * @param InParameter the parameter to query, e.g. -key=value
-		 * @returns the value of the parameter
-		 */
+		UE_DEPRECATED(5.8, "Use ProjectLauncher::CmdLineUtils version")
 		UE_API FString GetParameterValue( const FString& InParameter ) const;
 
-		/** 
-		 * Update the value of the given command line parameter
-		 * 
-		 * @param InParameter the parameter to modify, e.g. -key=value or -key=
-		 * @param InNewValue the new value to use
-		 * @returns true unless InParameter is not a key/value property
-		 */
+		UE_DEPRECATED(5.8, "Use ProjectLauncher::CmdLineUtils version")
 		UE_API bool UpdateParameterValue( const FString& InParameter, const FString& InNewValue );
 
-		/** 
-		 * Get the current state of the given command line parameter, allowing for user-changed values
-		 * 
-		 * @param InParameter the parameter to test, e.g. -key=value
-		 * @returns The parameter as it is now, allowing for user modifications, e.g. -key=value or -key=some_new_value
-		 */
+		UE_DEPRECATED(5.8, "Use ProjectLauncher::CmdLineUtils version")
 		UE_API FString GetFinalParameter( const FString& InParameter ) const;
 
 		/** 
 		 * Get the current command line
-		 * Extensions should generally go through this function instead of querying the Profile directlry
+		 * Extensions should generally go through this function instead of querying the Profile directly
 		 *
 		 * @returns the current command line string
 		 */
@@ -238,7 +204,7 @@ namespace ProjectLauncher
 
 		/** 
 		 * Update the current command line
-		 * Extensions should generally go through this function instead of modifying the Profile directlry
+		 * Extensions should generally go through this function instead of modifying the Profile directly
 		 * 
 		 * @param CommandLine the new command line
 		 */
@@ -374,9 +340,6 @@ namespace ProjectLauncher
 		UE_API void RequestTreeRefresh();
 
 	private:
-		UE_API bool IsParameterGroup( const FString& InParameter ) const;
-		UE_API bool TryRemoveParameterGroup( const FString& InParameter );
-
 		virtual int32 GetHeadingSortOrder() const { return 0; }
 
 
diff --git a/Engine/Plugins/Developer/ProjectLauncher/Source/ProjectLauncher/Public/Utils/LauncherCmdLineUtils.h b/Engine/Plugins/Developer/ProjectLauncher/Source/ProjectLauncher/Public/Utils/LauncherCmdLineUtils.h
new file mode 100644
--- /dev/null
+++ b/Engine/Plugins/Developer/ProjectLauncher/Source/ProjectLauncher/Public/Utils/LauncherCmdLineUtils.h
@@ -0,0 +1,72 @@
+// Copyright Epic Games, Inc. All Rights Reserved.
+
+#pragma once
+
+#include "CoreTypes.h"
+#include "Containers/UnrealString.h"
+
+#define UE_API PROJECTLAUNCHER_API
+
+namespace ProjectLauncher
+{
+	namespace CmdLineUtils
+	{
+		/**
+		 * Determine if the given parameter is on the command line
+ 		 * 
+		 * @param InParameter the parameter to test, e.g. -key=value or -param
+		 * @returns whether the parameter is being used
+		 */
+		UE_API bool IsParameterUsed( const FString& InCmdLine, const FString& InParameter );
+
+		/** 
+		 * Add or remove the given parameter on the command line
+		 * 
+		 * @param InParameter the parameter to add or remove, e.g. -key=value or -param
+		 * @param bUsed whether to add or remove the parameter
+		 */
+		UE_API void SetParameterUsed( FString& InOutCmdLine, const FString& InParameter, bool bUsed );
+
+		/** 
+		 * Add the given parameter to the command line
+		 * 
+		 * @param InParameter the parameter to add, e.g. -key=value or -param
+		 */
+		UE_API void AddParameter( FString& InOutCmdLine, const FString& InParameter );
+
+		/** 
+		 * Remove the given parameter from the command line
+		 * 
+		 * @param InParameter the parameter to remove, e.g. -key=value or -param
+		 * @returns true if the parameter was removed
+		 */
+		UE_API bool RemoveParameter( FString& InOutCmdLine, const FString& InParameter );
+
+		/** 
+		 * Retrieve the current value for the given command line parameter
+		 * 
+		 * @param InParameter the parameter to query, e.g. -key=value
+		 * @returns the value of the parameter
+		 */
+		UE_API FString GetParameterValue( const FString& InCmdLine, const FString& InParameter );
+
+		/** 
+		 * Update the value of the given command line parameter
+		 * 
+		 * @param InParameter the parameter to modify, e.g. -key=value or -key=
+		 * @param InNewValue the new value to use
+		 * @returns true unless InParameter is not a key/value property
+		 */
+		UE_API bool UpdateParameterValue( FString& InOutCmdLine, const FString& InParameter, const FString& InNewValue );
+
+		/** 
+		 * Get the current state of the given command line parameter, allowing for user-changed values
+		 * 
+		 * @param InParameter the parameter to test, e.g. -key=value
+		 * @returns The parameter as it is now, allowing for user modifications, e.g. -key=value or -key=some_new_value
+		 */
+		UE_API FString GetFinalParameter( const FString& InCmdLine, const FString& InParameter );
+	};
+};
+
+#undef UE_API
