Project Launcher extensions are now dynamically instantiated when the user enables them.
- extensions can now override IsAlwaysCreated to get the old behavior
- extension instances no longer manage their own "toggle" menu item (HasCustomExtensionMenu function has been removed)

#rnx
#Synced-CL 49557694

diff --git a/Engine/Plugins/Developer/ProjectLauncher/Source/CommonLaunchExtensions/Private/AdvancedCook/AdvancedCookLaunchExtension.cpp b/Engine/Plugins/Developer/ProjectLauncher/Source/CommonLaunchExtensions/Private/AdvancedCook/AdvancedCookLaunchExtension.cpp
--- a/Engine/Plugins/Developer/ProjectLauncher/Source/CommonLaunchExtensions/Private/AdvancedCook/AdvancedCookLaunchExtension.cpp
+++ b/Engine/Plugins/Developer/ProjectLauncher/Source/CommonLaunchExtensions/Private/AdvancedCook/AdvancedCookLaunchExtension.cpp
@@ -16,6 +16,7 @@
 
 FAdvancedCookLaunchExtensionInstance::FAdvancedCookLaunchExtensionInstance( FArgs& InArgs ) 
 	: ProjectLauncher::FLaunchExtensionInstance(InArgs)
+	, Owner(StaticCastSharedRef<FAdvancedCookLaunchExtension>(InArgs.Extension))
 {
 	MapOption = GetProfile()->GetCookedMaps().Num() > 0 ? EMapOption::Selected : EMapOption::Startup;
 
@@ -24,75 +25,11 @@ FAdvancedCookLaunchExtensionInstance::FAdvancedCookLaunchExtensionInstance( FArg
 }
 
 
-bool FAdvancedCookLaunchExtensionInstance::IsAvailable() const
-{
-	return (GetConfigBool(EConfig::PerProfile, TEXT("Enabled") ) || HasAdvancedCookOptions());
-}
-
-
-void FAdvancedCookLaunchExtensionInstance::MakeCustomExtensionSubmenu( FMenuBuilder& MenuBuilder )
-{
-	auto IsEnabled = [this]()
-	{
-		return GetConfigBool(EConfig::PerProfile, TEXT("Enabled") );
-	};
-
-	auto ToggleEnabled = [this,IsEnabled]()
-	{
-		SetConfigBool(EConfig::PerProfile, TEXT("Enabled"), !IsEnabled() );
-	};
-
-	auto CanToggle = [this,IsEnabled]()
-	{
-		return !HasAdvancedCookOptions();
-	};
-
-	auto GetToolTip = [this]()
-	{
-		if (HasAdvancedCookOptions() && IsAvailable())
-		{
-			return LOCTEXT("CannotDisable","This extension cannot be disabled while there are advanced cook options configured");
-		}
-
-		return FText::GetEmpty();
-	};
-
-
-	MenuBuilder.AddMenuEntry(
-		LOCTEXT("EnableAdvancedCookLabel", "Enable Advanced Cook"),
-		TAttribute<FText>::CreateLambda(GetToolTip),
-		FSlateIcon(),
-		FUIAction(
-			FExecuteAction::CreateLambda(ToggleEnabled),
-			FCanExecuteAction::CreateLambda(CanToggle),
-			FIsActionChecked::CreateSP(this, &FAdvancedCookLaunchExtensionInstance::IsAvailable)
-		),
-		NAME_None,
-		EUserInterfaceActionType::Check);
-}
-
-bool FAdvancedCookLaunchExtensionInstance::HasAdvancedCookOptions() const
-{
-	if (!GetProfile()->IsCookingUnversioned() || !GetProfile()->GetCookOptions().IsEmpty() || GetProfile()->GetCookedMaps().Num() > 0)
-	{
-		return true;
-	}
-
-	bool bHasDefaultCultures = (GetProfile()->GetCookedCultures().Num() == 1 && GetProfile()->GetCookedCultures()[0] == TEXT("en"));
-	if (!bHasDefaultCultures)
-	{
-		return true;
-	}
-
-	// we're not using any advanced cook options
-	return false;
-}
-
 void FAdvancedCookLaunchExtensionInstance::CustomizeTree( ProjectLauncher::FLaunchProfileTreeData& ProfileTreeData )
 {
 	auto IsVisible = [this]()
 	{
-		return IsAvailable() && (GetProfile()->GetCookMode() == ELauncherProfileCookModes::ByTheBook);
+		return (GetProfile()->GetCookMode() == ELauncherProfileCookModes::ByTheBook);
 	};
 
 	//ProfileTreeData.AddHeading( TEXT("Cooking"), LOCTEXT("CookingHeading", "Cook") )
@@ -428,6 +365,10 @@ void FAdvancedCookLaunchExtensionInstance::OnProjectChanged()
 	}
 }
 
+bool FAdvancedCookLaunchExtensionInstance::CanBeRemoved() const
+{
+	return !Owner->HasAdvancedCookOptions(GetProfile());
+}
 
 
 
@@ -446,5 +387,29 @@ FText FAdvancedCookLaunchExtension::GetDisplayName() const
 	return LOCTEXT("ExtensionName", "Advanced Cook Options");
 }
 
+bool FAdvancedCookLaunchExtension::IsCreatedByDefault( ILauncherProfileRef InProfile, TSharedRef<ProjectLauncher::FModel> InModel ) const
+{
+	return HasAdvancedCookOptions(InProfile);
+}
+
+bool FAdvancedCookLaunchExtension::HasAdvancedCookOptions(ILauncherProfileRef InProfile) const
+{
+	if (!InProfile->IsCookingUnversioned() || !InProfile->GetCookOptions().IsEmpty() || InProfile->GetCookedMaps().Num() > 0)
+	{
+		return true;
+	}
+
+	bool bHasDefaultCultures = (InProfile->GetCookedCultures().Num() == 1 && InProfile->GetCookedCultures()[0] == TEXT("en"));
+	if (!bHasDefaultCultures)
+	{
+		return true;
+	}
+
+	// we're not using any advanced cook options
+	return false;
+}
+
+
+
 
 #undef LOCTEXT_NAMESPACE
diff --git a/Engine/Plugins/Developer/ProjectLauncher/Source/CommonLaunchExtensions/Private/AdvancedCook/AdvancedCookLaunchExtension.h b/Engine/Plugins/Developer/ProjectLauncher/Source/CommonLaunchExtensions/Private/AdvancedCook/AdvancedCookLaunchExtension.h
--- a/Engine/Plugins/Developer/ProjectLauncher/Source/CommonLaunchExtensions/Private/AdvancedCook/AdvancedCookLaunchExtension.h
+++ b/Engine/Plugins/Developer/ProjectLauncher/Source/CommonLaunchExtensions/Private/AdvancedCook/AdvancedCookLaunchExtension.h
@@ -15,12 +15,13 @@ public:
 	FAdvancedCookLaunchExtensionInstance( FArgs& InArgs );
 	virtual ~FAdvancedCookLaunchExtensionInstance() = default;
 
-	virtual bool HasCustomExtensionMenu() const { return true; }
-	virtual void MakeCustomExtensionSubmenu( FMenuBuilder& MenuBuilder ) override;
-
 	virtual void CustomizeTree( ProjectLauncher::FLaunchProfileTreeData& ProfileTreeData ) override;
+	virtual bool CanBeRemoved() const override;
 
 protected:
+	TSharedRef<class FAdvancedCookLaunchExtension> Owner;
+
+
 	TArray<FString> CachedMapsToCook;
 	enum class EMapOption : uint8
 	{
@@ -62,9 +63,6 @@ protected:
 
 
 	virtual void OnProjectChanged() override;
-	bool HasAdvancedCookOptions() const;
-
-	bool IsAvailable() const;
 };
 
 
@@ -74,4 +72,7 @@ public:
 	virtual TSharedPtr<ProjectLauncher::FLaunchExtensionInstance> CreateInstanceForProfile( ProjectLauncher::FLaunchExtensionInstance::FArgs& InArgs ) override;
 	virtual const TCHAR* GetInternalName() const override;
 	virtual FText GetDisplayName() const override;
+	virtual bool IsCreatedByDefault( ILauncherProfileRef InProfile, TSharedRef<ProjectLauncher::FModel> InModel ) const override;
+
+	bool HasAdvancedCookOptions(ILauncherProfileRef InProfile) const;
 };
\ No newline at end of file
diff --git a/Engine/Plugins/Developer/ProjectLauncher/Source/CommonLaunchExtensions/Private/BootTest/BootTestLaunchExtension.h b/Engine/Plugins/Developer/ProjectLauncher/Source/CommonLaunchExtensions/Private/BootTest/BootTestLaunchExtension.h
--- a/Engine/Plugins/Developer/ProjectLauncher/Source/CommonLaunchExtensions/Private/BootTest/BootTestLaunchExtension.h
+++ b/Engine/Plugins/Developer/ProjectLauncher/Source/CommonLaunchExtensions/Private/BootTest/BootTestLaunchExtension.h
@@ -25,7 +25,7 @@ private:
 };
 
 
-class FBootTestLaunchExtension : public ProjectLauncher::FLaunchExtension
+class FBootTestLaunchExtension : public ProjectLauncher::FAutomatedTestLaunchExtension
 {
 public:
 	virtual TSharedPtr<ProjectLauncher::FLaunchExtensionInstance> CreateInstanceForProfile( ProjectLauncher::FLaunchExtensionInstance::FArgs& InArgs ) override;
diff --git a/Engine/Plugins/Developer/ProjectLauncher/Source/CommonLaunchExtensions/Private/DeprecatedProperties/DeprecatedPropertiesLaunchExtension.cpp b/Engine/Plugins/Developer/ProjectLauncher/Source/CommonLaunchExtensions/Private/DeprecatedProperties/DeprecatedPropertiesLaunchExtension.cpp
--- a/Engine/Plugins/Developer/ProjectLauncher/Source/CommonLaunchExtensions/Private/DeprecatedProperties/DeprecatedPropertiesLaunchExtension.cpp
+++ b/Engine/Plugins/Developer/ProjectLauncher/Source/CommonLaunchExtensions/Private/DeprecatedProperties/DeprecatedPropertiesLaunchExtension.cpp
@@ -1,6 +1,7 @@
 // Copyright Epic Games, Inc. All Rights Reserved.
 
 #include "DeprecatedProperties/DeprecatedPropertiesLaunchExtension.h"
+#include "Model/ProjectLauncherModel.h"
 #include "HAL/PlatformMisc.h"
 #include "Misc/CoreMisc.h"
 #include "Widgets/Shared/SCustomLaunchCombo.h"
@@ -14,39 +15,20 @@
 
 FDeprecatedPropertiesLaunchExtensionInstance::FDeprecatedPropertiesLaunchExtensionInstance( FArgs& InArgs ) 
 	: ProjectLauncher::FLaunchExtensionInstance(InArgs)
+	, Owner(StaticCastSharedRef<FDeprecatedPropertiesLaunchExtension>(InArgs.Extension))
 {
-	bEnabled = HasDeprecatedProperties();
 }
 
 
-bool FDeprecatedPropertiesLaunchExtensionInstance::HasDeprecatedProperties() const
-{
-	return	GetProfile()->GetSkipCookingEditorContent()				!=	GetDefaultProfile()->GetSkipCookingEditorContent() ||
-			GetProfile()->GetCookConfiguration()					!=	GetDefaultProfile()->GetCookConfiguration() ||
-			GetProfile()->GetPackageDirectory()						!=	GetDefaultProfile()->GetPackageDirectory() ||
-			GetProfile()->IsUsingIoStore()							!=	GetDefaultProfile()->IsUsingIoStore() ||
-			GetProfile()->GetReferenceContainerGlobalFileName()		!=	GetDefaultProfile()->GetReferenceContainerGlobalFileName() ||
-			GetProfile()->GetReferenceContainerCryptoKeysFileName()	!=	GetDefaultProfile()->GetReferenceContainerCryptoKeysFileName()
-		;
-}
-
 void FDeprecatedPropertiesLaunchExtensionInstance::CustomizeTree( ProjectLauncher::FLaunchProfileTreeData& ProfileTreeData )
 {
-	auto IsVisible = [this]()
-	{
-		return bEnabled;
-	};
-
 	AddDefaultHeading(ProfileTreeData)
 		.AddWidget( FText::GetEmpty(),
-			{
-				.IsVisible = IsVisible,
-			},
 			SNew(SBorder)
 			.BorderImage(FAppStyle::GetBrush("ToolPanel.LightGroupBorder"))
 			.BorderBackgroundColor_Lambda( [this]()
 			{
-				return HasDeprecatedProperties()
+				return Owner->HasDeprecatedProperties(GetProfile(), GetDefaultProfile())
 					? FAppStyle::Get().GetSlateColor("Colors.Warning")
 					: FAppStyle::Get().GetSlateColor("Colors.AccentGreen");
 			})
@@ -55,20 +37,19 @@ void FDeprecatedPropertiesLaunchExtensionInstance::CustomizeTree( ProjectLaunche
 				.ColorAndOpacity(FLinearColor::White)
 				.Text_Lambda( [this]() 
 				{ 
-					return HasDeprecatedProperties() 
+					return Owner->HasDeprecatedProperties(GetProfile(), GetDefaultProfile()) 
 						? LOCTEXT("DeprecatedItemsBanner", "These properties may be removed in future engine releases.") 
 						: LOCTEXT("NoDeprecatedItemsBanner", "These properties may be removed in future engine releases. You are not using any at the moment.");
 				})
 			]
 		)
-
+		
 		// ... http chunks ...
 		.AddBoolean( LOCTEXT("CreateHttpChunksLabel", "Create Http Chunk Install Data"),
 			{
 				.GetValue = [this]()				{ return GetProfile()->IsGenerateHttpChunkData(); },
 				.SetValue = [this](bool bValue)		{ GetProfile()->SetGenerateHttpChunkData(bValue); },
 				.GetDefaultValue = [this]()			{ return GetDefaultProfile()->IsGenerateHttpChunkData(); },
-				.IsVisible = IsVisible,
 			}
 		)
 		.AddDirectoryString( LOCTEXT("HttpChunkDirLabel", "Http Chunk Install Data Path"),
@@ -76,7 +57,6 @@ void FDeprecatedPropertiesLaunchExtensionInstance::CustomizeTree( ProjectLaunche
 				.GetValue = [this]()				{ return GetProfile()->GetHttpChunkDataDirectory(); },
 				.SetValue = [this](FString Value)	{ GetProfile()->SetHttpChunkDataDirectory(Value); },
 				.GetDefaultValue = [this]()			{ return GetDefaultProfile()->GetHttpChunkDataDirectory(); },
-				.IsVisible = IsVisible,
 				.IsEnabled = [this]()				{ return GetProfile()->IsGenerateHttpChunkData(); },
 			}
 		)
@@ -85,7 +65,6 @@ void FDeprecatedPropertiesLaunchExtensionInstance::CustomizeTree( ProjectLaunche
 				.GetValue = [this]()				{ return GetProfile()->GetHttpChunkDataReleaseName(); },
 				.SetValue = [this](FString Value)	{ GetProfile()->SetHttpChunkDataReleaseName(Value); },
 				.GetDefaultValue = [this]()			{ return GetDefaultProfile()->GetHttpChunkDataReleaseName(); },
-				.IsVisible = IsVisible,
 				.IsEnabled = [this]()				{ return GetProfile()->IsGenerateHttpChunkData(); },
 			}
 		)
@@ -96,14 +75,12 @@ void FDeprecatedPropertiesLaunchExtensionInstance::CustomizeTree( ProjectLaunche
 				.GetValue = [this]()				{ return GetProfile()->GetSkipCookingEditorContent(); },
 				.SetValue = [this](bool bValue)		{ GetProfile()->SetSkipCookingEditorContent(bValue); },
 				.GetDefaultValue = [this]()			{ return GetDefaultProfile()->GetSkipCookingEditorContent(); },
-				.IsVisible = IsVisible,
 			}
 		)
 		.AddWidget( LOCTEXT("CookerConfigLabel", "Cooker build configuration"), 
 			{
 				.IsDefault = [this]()				{ return GetProfile()->GetCookConfiguration() == GetDefaultProfile()->GetCookConfiguration(); },
 				.SetToDefault = [this]()			{ return GetProfile()->SetCookConfiguration( GetDefaultProfile()->GetCookConfiguration() ); },
-				.IsVisible = IsVisible,
 			},
 			SNew(SCustomLaunchLexToStringCombo<EBuildConfiguration>)
 			.OnSelectionChanged_Lambda( [this](EBuildConfiguration Value) { GetProfile()->SetCookConfiguration(Value); } )
@@ -115,8 +92,6 @@ void FDeprecatedPropertiesLaunchExtensionInstance::CustomizeTree( ProjectLaunche
 				.GetValue = [this]()				{ return GetProfile()->GetPackageDirectory(); },
 				.SetValue = [this](FString Value)	{ GetProfile()->SetPackageDirectory(Value); },
 				.GetDefaultValue = [this]()			{ return GetDefaultProfile()->GetPackageDirectory(); },
-				.IsVisible = IsVisible,
-				.IsEnabled = IsVisible,
 			}
 		)
 		.AddBoolean( LOCTEXT("IoStoreLabel", "Use container files for optimized loading (I/O Store)"),
@@ -124,7 +99,6 @@ void FDeprecatedPropertiesLaunchExtensionInstance::CustomizeTree( ProjectLaunche
 				.GetValue = [this]()				{ return GetProfile()->IsUsingIoStore(); },
 				.SetValue = [this](bool bValue)		{ GetProfile()->SetUseIoStore(bValue); },
 				.GetDefaultValue = [this]()			{ return GetDefaultProfile()->IsUsingIoStore(); },
-				.IsVisible = IsVisible,
 			}
 		)
 		.AddFileString( LOCTEXT("IoStoreRefContainerLabel", "Optional I/O Store Reference Chunk Database"),
@@ -132,7 +106,6 @@ void FDeprecatedPropertiesLaunchExtensionInstance::CustomizeTree( ProjectLaunche
 				.GetValue = [this]()				{ return GetProfile()->GetReferenceContainerGlobalFileName(); },
 				.SetValue = [this](FString Value)	{ GetProfile()->SetReferenceContainerGlobalFileName(Value); },
 				.GetDefaultValue = [this]()			{ return GetDefaultProfile()->GetReferenceContainerGlobalFileName(); },
-				.IsVisible = IsVisible,
 				.IsEnabled = [this]()				{ return GetProfile()->IsUsingIoStore(); },
 			},
 			TEXT("global.utoc files|global.utoc")
@@ -142,7 +115,6 @@ void FDeprecatedPropertiesLaunchExtensionInstance::CustomizeTree( ProjectLaunche
 				.GetValue = [this]()				{ return GetProfile()->GetReferenceContainerCryptoKeysFileName(); },
 				.SetValue = [this](FString Value)	{ GetProfile()->SetReferenceContainerCryptoKeysFileName(Value); },
 				.GetDefaultValue = [this]()			{ return GetDefaultProfile()->GetReferenceContainerCryptoKeysFileName(); },
-				.IsVisible = IsVisible,
 				.IsEnabled = [this]()				{ return GetProfile()->IsUsingIoStore(); },
 			},
 			TEXT("crypto.json files|crypto.json")
@@ -150,38 +122,30 @@ void FDeprecatedPropertiesLaunchExtensionInstance::CustomizeTree( ProjectLaunche
 	;
 }
 
-void FDeprecatedPropertiesLaunchExtensionInstance::MakeCustomExtensionSubmenu( FMenuBuilder& MenuBuilder )
+bool FDeprecatedPropertiesLaunchExtensionInstance::CanBeRemoved() const
 {
-	MenuBuilder.AddMenuEntry(
-		GetExtension()->GetDisplayName(),
-		FText::GetEmpty(),
-		FSlateIcon(),
-		FUIAction( 
-			FExecuteAction::CreateLambda( [this]() { bEnabled = !bEnabled; } ),
-			FCanExecuteAction(),
-			FIsActionChecked::CreateLambda( [this]() { return bEnabled; } )
-			),
-		NAME_None,
-		EUserInterfaceActionType::Check
-	);
+	// cannot remove the extension while there are deprecated properties set
+	return !Owner->HasDeprecatedProperties(GetProfile(), GetDefaultProfile());
+}
 
+void FDeprecatedPropertiesLaunchExtension::MakeCustomExtensionSubmenu(FMenuBuilder& MenuBuilder, ILauncherProfileRef InProfile, TSharedRef<ProjectLauncher::FModel> InModel)
+{
 	// putting this here for now because it isn't serialized with the profile. @todo: should have a better home!
 	MenuBuilder.AddMenuEntry(
 		LOCTEXT("BuildUATLabel", "Build UAT"),
 		LOCTEXT("BuildUATToolTip", "Build UAT before launching"),
 		FSlateIcon(),
 		FUIAction(
-			FExecuteAction::CreateLambda( [this]() { GetProfile()->SetBuildUAT(!GetProfile()->IsBuildingUAT()); } ),
+			FExecuteAction::CreateLambda( [this,InProfile]() { InProfile->SetBuildUAT(!InProfile->IsBuildingUAT()); } ),
 			FCanExecuteAction(),
-			FIsActionChecked::CreateLambda( [this]() { return GetProfile()->IsBuildingUAT(); } )
+			FIsActionChecked::CreateLambda( [this,InProfile]() { return InProfile->IsBuildingUAT(); } )
 		),
 		NAME_None,
 		EUserInterfaceActionType::Check
 	);
 }
 
-
-void FDeprecatedPropertiesLaunchExtensionInstance::GetExtensionsMenuEntry( FExtensionsMenuEntry& MenuEntry ) const
+void FDeprecatedPropertiesLaunchExtension::GetExtensionsMenuEntry( FExtensionsMenuEntry& MenuEntry ) const
 {
 	MenuEntry = FExtensionsMenuEntry::Deprecated;	
 }
@@ -204,4 +168,24 @@ FText FDeprecatedPropertiesLaunchExtension::GetDisplayName() const
 }
 
 
+bool FDeprecatedPropertiesLaunchExtension::IsCreatedByDefault( ILauncherProfileRef InProfile, TSharedRef<ProjectLauncher::FModel> InModel ) const
+{
+	ILauncherProfileRef DefaultProfile = InModel->GetDefaultCustomLaunchProfile();
+	return HasDeprecatedProperties(InProfile, DefaultProfile);
+}
+
+bool FDeprecatedPropertiesLaunchExtension::HasDeprecatedProperties(ILauncherProfileRef InProfile, ILauncherProfileRef InDefaultProfile) const
+{
+	return	InProfile->GetSkipCookingEditorContent()				!=	InDefaultProfile->GetSkipCookingEditorContent() ||
+			InProfile->GetCookConfiguration()						!=	InDefaultProfile->GetCookConfiguration() ||
+			InProfile->GetPackageDirectory()						!=	InDefaultProfile->GetPackageDirectory() ||
+			InProfile->IsUsingIoStore()								!=	InDefaultProfile->IsUsingIoStore() ||
+			InProfile->GetReferenceContainerGlobalFileName()		!=	InDefaultProfile->GetReferenceContainerGlobalFileName() ||
+			InProfile->GetReferenceContainerCryptoKeysFileName()	!=	InDefaultProfile->GetReferenceContainerCryptoKeysFileName()
+		;
+
+}
+
+
+
 #undef LOCTEXT_NAMESPACE
diff --git a/Engine/Plugins/Developer/ProjectLauncher/Source/CommonLaunchExtensions/Private/DeprecatedProperties/DeprecatedPropertiesLaunchExtension.h b/Engine/Plugins/Developer/ProjectLauncher/Source/CommonLaunchExtensions/Private/DeprecatedProperties/DeprecatedPropertiesLaunchExtension.h
--- a/Engine/Plugins/Developer/ProjectLauncher/Source/CommonLaunchExtensions/Private/DeprecatedProperties/DeprecatedPropertiesLaunchExtension.h
+++ b/Engine/Plugins/Developer/ProjectLauncher/Source/CommonLaunchExtensions/Private/DeprecatedProperties/DeprecatedPropertiesLaunchExtension.h
@@ -12,15 +12,10 @@ public:
 	virtual ~FDeprecatedPropertiesLaunchExtensionInstance() = default;
 
 	virtual void CustomizeTree( ProjectLauncher::FLaunchProfileTreeData& ProfileTreeData ) override;
-
-	virtual bool HasCustomExtensionMenu() const { return true; }
-	virtual void MakeCustomExtensionSubmenu( FMenuBuilder& MenuBuilder ) override;
-	virtual void GetExtensionsMenuEntry( FExtensionsMenuEntry& MenuEntry ) const override;
+	virtual bool CanBeRemoved() const override;
 
 protected:
-	bool bEnabled = false;
-	bool bHasBeenShown = false;
-	bool HasDeprecatedProperties() const;
+	TSharedRef<class FDeprecatedPropertiesLaunchExtension> Owner;
 };
 
 
@@ -30,4 +25,9 @@ public:
 	virtual TSharedPtr<ProjectLauncher::FLaunchExtensionInstance> CreateInstanceForProfile( ProjectLauncher::FLaunchExtensionInstance::FArgs& InArgs ) override;
 	virtual const TCHAR* GetInternalName() const override;
 	virtual FText GetDisplayName() const override;
+	virtual bool IsCreatedByDefault( ILauncherProfileRef InProfile, TSharedRef<ProjectLauncher::FModel> InModel ) const override;
+	virtual void GetExtensionsMenuEntry( FExtensionsMenuEntry& MenuEntry ) const override;
+	virtual void MakeCustomExtensionSubmenu(FMenuBuilder& MenuBuilder, ILauncherProfileRef InProfile, TSharedRef<ProjectLauncher::FModel> InModel) override;
+	
+	bool HasDeprecatedProperties(ILauncherProfileRef InProfile, ILauncherProfileRef InDefaultProfile) const;
 };
\ No newline at end of file
diff --git a/Engine/Plugins/Developer/ProjectLauncher/Source/CommonLaunchExtensions/Private/Globals/GlobalsLaunchExtension.h b/Engine/Plugins/Developer/ProjectLauncher/Source/CommonLaunchExtensions/Private/Globals/GlobalsLaunchExtension.h
--- a/Engine/Plugins/Developer/ProjectLauncher/Source/CommonLaunchExtensions/Private/Globals/GlobalsLaunchExtension.h
+++ b/Engine/Plugins/Developer/ProjectLauncher/Source/CommonLaunchExtensions/Private/Globals/GlobalsLaunchExtension.h
@@ -30,4 +30,5 @@ public:
 	virtual TSharedPtr<ProjectLauncher::FLaunchExtensionInstance> CreateInstanceForProfile( ProjectLauncher::FLaunchExtensionInstance::FArgs& InArgs ) override;
 	virtual const TCHAR* GetInternalName() const override;
 	virtual FText GetDisplayName() const override;
+	virtual bool IsAlwaysCreated(ILauncherProfileRef InProfile, TSharedRef<ProjectLauncher::FModel> InModel) const { return true; }
 };
\ No newline at end of file
diff --git a/Engine/Plugins/Developer/ProjectLauncher/Source/CommonLaunchExtensions/Private/Insights/InsightsLaunchExtension.h b/Engine/Plugins/Developer/ProjectLauncher/Source/CommonLaunchExtensions/Private/Insights/InsightsLaunchExtension.h
--- a/Engine/Plugins/Developer/ProjectLauncher/Source/CommonLaunchExtensions/Private/Insights/InsightsLaunchExtension.h
+++ b/Engine/Plugins/Developer/ProjectLauncher/Source/CommonLaunchExtensions/Private/Insights/InsightsLaunchExtension.h
@@ -38,4 +38,5 @@ public:
 	virtual TSharedPtr<ProjectLauncher::FLaunchExtensionInstance> CreateInstanceForProfile( ProjectLauncher::FLaunchExtensionInstance::FArgs& InArgs ) override;
 	virtual const TCHAR* GetInternalName() const override;
 	virtual FText GetDisplayName() const override;
+	virtual bool IsAlwaysCreated(ILauncherProfileRef InProfile, TSharedRef<ProjectLauncher::FModel> InModel) const override { return true; }
 };
\ No newline at end of file
diff --git a/Engine/Plugins/Developer/ProjectLauncher/Source/CommonLaunchExtensions/Private/ProfileWizard/ProfileWizardLaunchExtension.cpp b/Engine/Plugins/Developer/ProjectLauncher/Source/CommonLaunchExtensions/Private/ProfileWizard/ProfileWizardLaunchExtension.cpp
--- a/Engine/Plugins/Developer/ProjectLauncher/Source/CommonLaunchExtensions/Private/ProfileWizard/ProfileWizardLaunchExtension.cpp
+++ b/Engine/Plugins/Developer/ProjectLauncher/Source/CommonLaunchExtensions/Private/ProfileWizard/ProfileWizardLaunchExtension.cpp
@@ -6,21 +6,15 @@
 #define LOCTEXT_NAMESPACE "FProfileWizardLaunchExtensionInstance"
 
 
-FProfileWizardLaunchExtensionInstance::FProfileWizardLaunchExtensionInstance( FArgs& InArgs ) 
-	: ProjectLauncher::FLaunchExtensionInstance(InArgs)
+void FProfileWizardLaunchExtension::MakeCustomExtensionSubmenu( FMenuBuilder& MenuBuilder, ILauncherProfileRef InProfile, TSharedRef<ProjectLauncher::FModel> InModel )
 {
-}
-
-
-void FProfileWizardLaunchExtensionInstance::MakeCustomExtensionSubmenu( FMenuBuilder& MenuBuilder )
-{
-	auto CreateSubMenu = [this]( FMenuBuilder& MenuBuilder )
+	auto CreateSubMenu = [this, InModel]( FMenuBuilder& MenuBuilder )
 	{
-		for (const ILauncherProfileWizardPtr& ProfileWizard : GetModel()->GetProfileManager()->GetProfileWizards())
+		for (const ILauncherProfileWizardPtr& ProfileWizard : InModel->GetProfileManager()->GetProfileWizards())
 		{
-			auto RunWizard = [this, ProfileWizard]()
+			auto RunWizard = [this, InModel, ProfileWizard]()
 			{
-				ProfileWizard->HandleCreateLauncherProfile( GetModel()->GetProfileManager() );
+				ProfileWizard->HandleCreateLauncherProfile( InModel->GetProfileManager() );
 			};
 
 			MenuBuilder.AddMenuEntry(
@@ -32,25 +26,17 @@ void FProfileWizardLaunchExtensionInstance::MakeCustomExtensionSubmenu( FMenuBui
 		}
 	};
 
-	MenuBuilder.AddSubMenu( GetExtension()->GetDisplayName(), FText::GetEmpty(), FNewMenuDelegate::CreateLambda(CreateSubMenu));
-}
-
-
-void FProfileWizardLaunchExtensionInstance::GetExtensionsMenuEntry( FExtensionsMenuEntry& MenuEntry ) const
-{
-	MenuEntry = FExtensionsMenuEntry::Deprecated;	
+	if (InModel->GetProfileManager()->GetProfileWizards().Num() > 0)
+	{
+		MenuBuilder.AddSubMenu( GetDisplayName(), FText::GetEmpty(), FNewMenuDelegate::CreateLambda(CreateSubMenu));
+	}
 }
 
 
 
 TSharedPtr<ProjectLauncher::FLaunchExtensionInstance> FProfileWizardLaunchExtension::CreateInstanceForProfile( ProjectLauncher::FLaunchExtensionInstance::FArgs& InArgs )
 {
-	if (InArgs.Model->GetProfileManager()->GetProfileWizards().Num() > 0)
-	{
-		return MakeShared<FProfileWizardLaunchExtensionInstance>(InArgs);
-	}
-
-	// there's no profile wizards - don't need to create the extension
+	// this extension doesn't need instantation
 	return nullptr;
 }
 
@@ -67,4 +53,9 @@ FText FProfileWizardLaunchExtension::GetDisplayName() const
 }
 
 
+void FProfileWizardLaunchExtension::GetExtensionsMenuEntry( FExtensionsMenuEntry& MenuEntry ) const
+{
+	MenuEntry = FExtensionsMenuEntry::Deprecated;	
+}
+
 #undef LOCTEXT_NAMESPACE
diff --git a/Engine/Plugins/Developer/ProjectLauncher/Source/CommonLaunchExtensions/Private/ProfileWizard/ProfileWizardLaunchExtension.h b/Engine/Plugins/Developer/ProjectLauncher/Source/CommonLaunchExtensions/Private/ProfileWizard/ProfileWizardLaunchExtension.h
--- a/Engine/Plugins/Developer/ProjectLauncher/Source/CommonLaunchExtensions/Private/ProfileWizard/ProfileWizardLaunchExtension.h
+++ b/Engine/Plugins/Developer/ProjectLauncher/Source/CommonLaunchExtensions/Private/ProfileWizard/ProfileWizardLaunchExtension.h
@@ -4,24 +4,14 @@
 
 #include "Extension/LaunchExtension.h"
 
-class FProfileWizardLaunchExtensionInstance : public ProjectLauncher::FLaunchExtensionInstance
-{
-public:
-	FProfileWizardLaunchExtensionInstance( FArgs& InArgs );
-	virtual ~FProfileWizardLaunchExtensionInstance() = default;
-
-	virtual bool HasCustomExtensionMenu() const { return true; }
-	virtual void MakeCustomExtensionSubmenu( FMenuBuilder& MenuBuilder ) override;
-	virtual void GetExtensionsMenuEntry( FExtensionsMenuEntry& MenuEntry ) const override;
-
-protected:
-};
-
-
 class FProfileWizardLaunchExtension : public ProjectLauncher::FLaunchExtension
 {
 public:
 	virtual TSharedPtr<ProjectLauncher::FLaunchExtensionInstance> CreateInstanceForProfile( ProjectLauncher::FLaunchExtensionInstance::FArgs& InArgs ) override;
 	virtual const TCHAR* GetInternalName() const override;
 	virtual FText GetDisplayName() const override;
-};
\ No newline at end of file
+	virtual bool IsAlwaysCreated(ILauncherProfileRef InProfile, TSharedRef<ProjectLauncher::FModel> InModel) const override { return true; }
+	virtual void GetExtensionsMenuEntry( FExtensionsMenuEntry& MenuEntry ) const override;
+	virtual void MakeCustomExtensionSubmenu(FMenuBuilder& MenuBuilder, ILauncherProfileRef InProfile, TSharedRef<ProjectLauncher::FModel> InModel) override;
+
+};
diff --git a/Engine/Plugins/Developer/ProjectLauncher/Source/ProjectLauncher/Private/Extension/LaunchExtension.cpp b/Engine/Plugins/Developer/ProjectLauncher/Source/ProjectLauncher/Private/Extension/LaunchExtension.cpp
--- a/Engine/Plugins/Developer/ProjectLauncher/Source/ProjectLauncher/Private/Extension/LaunchExtension.cpp
+++ b/Engine/Plugins/Developer/ProjectLauncher/Source/ProjectLauncher/Private/Extension/LaunchExtension.cpp
@@ -20,6 +20,7 @@ namespace ProjectLauncher
 {
 	static TArray<TSharedPtr<FLaunchExtension>> GExtensions;
 	static TMap<ILauncherProfilePtr, TArray<TSharedPtr<FLaunchExtensionInstance>>> GExtensionInstances;
+	static TMap<TSharedPtr<FLaunchExtension>,TArray<TSharedPtr<FLaunchExtensionInstance>>> GInstancesPerExtension;
 	static FCriticalSection GExtensionsCS;
 	static std::atomic<bool> GExtensionsDestroyed = false;
 
@@ -30,6 +31,7 @@ namespace ProjectLauncher
 		{
 			GExtensions.Reset();
 			GExtensionInstances.Reset();
+			GInstancesPerExtension.Reset();
 
 			GExtensionsDestroyed = true;
 		}
@@ -42,6 +44,7 @@ namespace ProjectLauncher
 
 		FScopeLock Lock(&GExtensionsCS);
 		GExtensions.Add(Extension);
+		GInstancesPerExtension.Add(Extension);
 	}
 
 	void UnregisterExtension( TSharedRef<FLaunchExtension> Extension )
@@ -50,6 +53,16 @@ namespace ProjectLauncher
 		{
 			FScopeLock Lock(&GExtensionsCS);
 			GExtensions.Remove(Extension);
+			GInstancesPerExtension.Remove(Extension);
+
+			for (TTuple<ILauncherProfilePtr, TArray<TSharedPtr<FLaunchExtensionInstance>>>& Pair : GExtensionInstances)
+			{
+				TArray<TSharedPtr<FLaunchExtensionInstance>>& ExtensionInstances = Pair.Value;
+				ExtensionInstances.RemoveAll( [Extension](TSharedPtr<FLaunchExtensionInstance> ExtensionInstance)
+				{
+					return ExtensionInstance->GetExtension() == Extension;
+				});
+			}
 		}
 	}
 
@@ -160,11 +173,34 @@ namespace ProjectLauncher
 	{
 		FScopeLock Lock(&GExtensionsCS);
 		GExtensionInstances.Remove(Profile.ToSharedPtr());
+
+		for (TTuple<TSharedPtr<FLaunchExtension>,TArray<TSharedPtr<FLaunchExtensionInstance>>>& Pair : GInstancesPerExtension)
+		{
+			TArray<TSharedPtr<FLaunchExtensionInstance>>& ExtensionInstances = Pair.Value;
+			ExtensionInstances.RemoveAll( [Profile](TSharedPtr<FLaunchExtensionInstance> ExtensionInstance)
+			{
+				return ExtensionInstance->GetProfile() == Profile;
+			});
+		}
 	}
 
 
 
-	TArray<TSharedPtr<FLaunchExtensionInstance>> FLaunchExtension::CreateExtensionInstancesForProfile( ILauncherProfileRef InProfile, TSharedRef<FModel> InModel, FLaunchProfileTreeData& OwnerTreeData )
+	TArray<TSharedRef<FLaunchExtension>> FLaunchExtension::GetExtensions()
+	{
+		FScopeLock Lock(&GExtensionsCS);
+
+		TArray<TSharedRef<FLaunchExtension>> Result;
+		for ( const TSharedPtr<FLaunchExtension>& Extension : GExtensions)
+		{
+			Result.Add(Extension.ToSharedRef());
+		}
+
+		return Result;
+	}
+
+
+	TArray<TSharedPtr<FLaunchExtensionInstance>> FLaunchExtension::CreateExtensionInstancesForProfile( ILauncherProfileRef InProfile, TSharedRef<FModel> InModel, FLaunchProfileTreeData& InOwnerTreeData )
 	{
 		check(!GExtensionsDestroyed);
 
@@ -185,36 +221,135 @@ namespace ProjectLauncher
 			Extensions = GExtensions;
 		}
 
+		// attempt to instantiate all extensions
+		for (TSharedPtr<FLaunchExtension> Extension : Extensions)
+		{
+			const bool* ValuePtr = InProfile->GetCustomBoolProperties().Find( FString::Printf(TEXT("%s.Instantiated"), Extension->GetInternalName() ) );
+			bool bEnabled = ValuePtr && (*ValuePtr);
+
+			if (Extension->IsAlwaysCreated(InProfile, InModel) || Extension->IsCreatedByDefault(InProfile, InModel) || bEnabled)
+			{
+				CreateExtensionInstance(Extension.ToSharedRef(), InProfile, InModel, InOwnerTreeData);
+			}
+		}
+
+		// return all extensions
+		{
+			FScopeLock Lock(&GExtensionsCS);
+			const TArray<TSharedPtr<FLaunchExtensionInstance>>* ExistingInstances = GExtensionInstances.Find(InProfile.ToSharedPtr());
+			if (ExistingInstances != nullptr)
+			{
+				return (*ExistingInstances);
+			}
+		}
+
+		return TArray<TSharedPtr<FLaunchExtensionInstance>>();
+	}
+
+	TSharedPtr<FLaunchExtensionInstance> FLaunchExtension::CreateExtensionInstance( TSharedRef<FLaunchExtension> InExtension, ILauncherProfileRef InProfile, TSharedRef<FModel> InModel, FLaunchProfileTreeData& InOwnerTreeData )
+	{
+		check(!GExtensionsDestroyed);
+
+		TSharedPtr<FLaunchExtensionInstance> ExistingInstance = GetExtensionInstance(InExtension, InProfile);
+		if (ExistingInstance.IsValid())
+		{
+			return ExistingInstance;
+		}
+
+		if (!InExtension->CanBeCreated(InProfile, InModel))
+		{
+			return nullptr;
+		}
+
 		// look up the default profile
 		ILauncherProfileRef DefaultProfile = InModel->IsBasicLaunchProfile(InProfile) ? InModel->GetDefaultBasicLaunchProfile() : InModel->GetDefaultCustomLaunchProfile();
 
-		// attempt to instantiate all extensions
-		TArray<TSharedPtr<FLaunchExtensionInstance>> Result;
-		for (TSharedPtr<FLaunchExtension> Extension : Extensions)
+		
+
+		FLaunchExtensionInstance::FArgs Args
+		{
+			.Profile = InProfile,
+			.DefaultProfile = DefaultProfile,
+			.Model = InModel,
+			.Extension = InExtension,
+			.OwnerTreeData = InOwnerTreeData,
+		};
+		TSharedPtr<FLaunchExtensionInstance> ExtensionInstance = InExtension->CreateInstanceForProfile(Args);
+
+		if (ExtensionInstance.IsValid())
 		{
-			FLaunchExtensionInstance::FArgs Args
+			ExtensionInstance->InternalInitialize();
+
+			// if this extension needs to be created dynamically, remember we have created it for next time and intialize it
+			if (!InExtension->IsAlwaysCreated(InProfile, InModel))
 			{
-				.Profile = InProfile,
-				.DefaultProfile = DefaultProfile,
-				.Model = InModel,
-				.Extension = Extension.ToSharedRef(),
-				.OwnerTreeData = OwnerTreeData,
-			};
-
-			TSharedPtr<FLaunchExtensionInstance> ExtensionInstance = Extension->CreateInstanceForProfile(Args);
-			if (ExtensionInstance.IsValid())
+				InProfile->GetCustomBoolProperties().Add( FString::Printf(TEXT("%s.Instantiated"), InExtension->GetInternalName() ), true );
+				ExtensionInstance->OnAdded();
+			}
+
+			// add it to the global caches
 			{
-				ExtensionInstance->InternalInitialize();
-				Result.Add(ExtensionInstance);
+				FScopeLock Lock(&GExtensionsCS);
+				if (!GExtensionInstances.Contains(InProfile.ToSharedPtr()))
+				{
+					GExtensionInstances.Add(InProfile.ToSharedPtr());
+				}
+				GExtensionInstances.FindChecked(InProfile.ToSharedPtr()).Add(ExtensionInstance);
+				GInstancesPerExtension.FindChecked(InExtension.ToSharedPtr()).Add(ExtensionInstance);
 			}
 		}
 
-		// record the extension instances for the this profile
+		return ExtensionInstance;
+	}
+
+	void FLaunchExtension::DestroyExtensionInstance( TSharedRef<FLaunchExtension> InExtension, ILauncherProfileRef InProfile, TSharedRef<FModel> InModel)
+	{
+		check(!GExtensionsDestroyed);
+
+		// can't destroy auto-instantiated extensions
+		check(!InExtension->IsAlwaysCreated(InProfile, InModel));
+
+		// remember it is no longer created
+		InProfile->GetCustomBoolProperties().Add( FString::Printf(TEXT("%s.Instantiated"), InExtension->GetInternalName() ), false );
+
+		// remove it from the global caches
+		TSharedPtr<FLaunchExtensionInstance> ExistingInstance = GetExtensionInstance(InExtension, InProfile);
+		if (ExistingInstance.IsValid())
 		{
+			ExistingInstance->OnRemoved();
+
 			FScopeLock Lock(&GExtensionsCS);
-			GExtensionInstances.Add(InProfile.ToSharedPtr(), Result);
+			GExtensionInstances.FindChecked(InProfile.ToSharedPtr()).Remove(ExistingInstance);
+			GInstancesPerExtension.FindChecked(InExtension.ToSharedPtr()).Remove(ExistingInstance);
 		}
-		return Result;
+	}
+
+	TSharedPtr<FLaunchExtensionInstance> FLaunchExtension::GetExtensionInstance(TSharedRef<FLaunchExtension> InExtension, ILauncherProfileRef InProfile)
+	{
+		check(!GExtensionsDestroyed);
+
+		FScopeLock Lock(&GExtensionsCS);
+		
+		const TArray<TSharedPtr<FLaunchExtensionInstance>>* InstancesForProfile = GExtensionInstances.Find(InProfile.ToSharedPtr());
+		const TArray<TSharedPtr<FLaunchExtensionInstance>>* InstancesForExtension = GInstancesPerExtension.Find(InExtension.ToSharedPtr());
+		if (InstancesForProfile != nullptr && InstancesForExtension != nullptr)
+		{
+			for (const TSharedPtr<FLaunchExtensionInstance>& ExtensionInstance : *InstancesForProfile)
+			{
+				if (InstancesForExtension->Contains(ExtensionInstance))
+				{
+					return ExtensionInstance;
+				}
+			}
+		}
+
+		return nullptr;
+	}
+
+
+	void FLaunchExtension::GetExtensionsMenuEntry( FExtensionsMenuEntry& MenuEntry ) const
+	{
+		MenuEntry = FExtensionsMenuEntry::Default;
 	}
 
 
@@ -227,53 +362,63 @@ namespace ProjectLauncher
 		static const FText DefaultSectionDisplayName = LOCTEXT("CommonSectionName","General Extensions");
 	}
 
-	const FLaunchExtensionInstance::FExtensionsMenuEntry FLaunchExtensionInstance::FExtensionsMenuEntry::None =
+	const FLaunchExtension::FExtensionsMenuEntry FLaunchExtension::FExtensionsMenuEntry::None =
 	{
-		.Type = FLaunchExtensionInstance::FExtensionsMenuEntry::Type_None,
+		.Type = FLaunchExtension::FExtensionsMenuEntry::Type_None,
 	};
 
-	const FLaunchExtensionInstance::FExtensionsMenuEntry FLaunchExtensionInstance::FExtensionsMenuEntry::Default =
+	const FLaunchExtension::FExtensionsMenuEntry FLaunchExtension::FExtensionsMenuEntry::Default =
 	{
 		.SectionName = DefaultSectionName,
 		.SectionDisplayName = DefaultSectionDisplayName,
-		.Type = FLaunchExtensionInstance::FExtensionsMenuEntry::Type_Section,
+		.Type = FLaunchExtension::FExtensionsMenuEntry::Type_Section,
 	};
 
-	const FLaunchExtensionInstance::FExtensionsMenuEntry FLaunchExtensionInstance::FExtensionsMenuEntry::OwnSection =
+	const FLaunchExtension::FExtensionsMenuEntry FLaunchExtension::FExtensionsMenuEntry::OwnSection =
 	{
-		.Type = FLaunchExtensionInstance::FExtensionsMenuEntry::Type_Section,
+		.Type = FLaunchExtension::FExtensionsMenuEntry::Type_Section,
 	};
 
-	const FLaunchExtensionInstance::FExtensionsMenuEntry FLaunchExtensionInstance::FExtensionsMenuEntry::Deprecated =
+	const FLaunchExtension::FExtensionsMenuEntry FLaunchExtension::FExtensionsMenuEntry::Deprecated =
 	{
 		.SectionName = DefaultSectionName,
 		.SubmenuName = TEXT("Deprecated"),
 		.SectionDisplayName = DefaultSectionDisplayName,
 		.SubmenuDisplayName = LOCTEXT("DeprecatedSubmenu", "Deprecated"),
-		.Type = FLaunchExtensionInstance::FExtensionsMenuEntry::Type_SubMenu,
+		.Type = FLaunchExtension::FExtensionsMenuEntry::Type_SubMenu,
 	};
 
-	const FLaunchExtensionInstance::FExtensionsMenuEntry FLaunchExtensionInstance::FExtensionsMenuEntry::AutomatedTests =
+	const FLaunchExtension::FExtensionsMenuEntry FLaunchExtension::FExtensionsMenuEntry::AutomatedTests =
 	{
 		.SectionName = DefaultSectionName,
 		.SubmenuName = TEXT("AutomatedTests"),
 		.SectionDisplayName = DefaultSectionDisplayName,
 		.SubmenuDisplayName = LOCTEXT("AutotestMenuSubMenu", "Automated Tests"),
-		.Type = FLaunchExtensionInstance::FExtensionsMenuEntry::Type_SubMenu,
+		.Type = FLaunchExtension::FExtensionsMenuEntry::Type_SubMenu,
 	};
 
 
-	const FLaunchExtensionInstance::FExtensionsMenuEntry FLaunchExtensionInstance::FExtensionsMenuEntry::UATCommands =
+	const FLaunchExtension::FExtensionsMenuEntry FLaunchExtension::FExtensionsMenuEntry::UATCommands =
 	{
 		.SectionName = DefaultSectionName,
 		.SubmenuName = TEXT("UATCommands"),
 		.SectionDisplayName = DefaultSectionDisplayName,
-		.SubmenuDisplayName = LOCTEXT("UATCommandsSubMenu", "More Buid Steps"),
-		.Type = FLaunchExtensionInstance::FExtensionsMenuEntry::Type_SubMenu,
+		.SubmenuDisplayName = LOCTEXT("UATCommandsSubMenu", "More Build Steps"),
+		.Type = FLaunchExtension::FExtensionsMenuEntry::Type_SubMenu,
 	};
 
 
 
+	void FAutomatedTestLaunchExtension::GetExtensionsMenuEntry( FExtensionsMenuEntry& MenuEntry ) const
+	{
+		MenuEntry = FExtensionsMenuEntry::AutomatedTests;
+	}
+
+	void FUATCommandLaunchExtension::GetExtensionsMenuEntry( FExtensionsMenuEntry& MenuEntry ) const
+	{
+		MenuEntry = FExtensionsMenuEntry::UATCommands;
+	}
+
 
 
 
@@ -627,11 +772,6 @@ namespace ProjectLauncher
 		return ProfileTreeData.AddHeading(Extension->GetInternalName(), OverrideDisplayName, GetHeadingSortOrder() );
 	}
 
-	void FLaunchExtensionInstance::GetExtensionsMenuEntry( FExtensionsMenuEntry& MenuEntry ) const
-	{
-		MenuEntry = FExtensionsMenuEntry::Default;
-	}
-
 
 
 
@@ -650,37 +790,22 @@ namespace ProjectLauncher
 		}
 	}
 
-	void FUATCommandLaunchExtensionInstance::MakeCustomExtensionSubmenu( FMenuBuilder& MenuBuilder )
+	void FUATCommandLaunchExtensionInstance::OnAdded()
+	{
+		ILauncherProfileUATCommandRef UATCommand = GetProfile()->FindOrAddUATCommand(*GetUATCommandInternalName());
+		OnUATCommandAdded(UATCommand);
+		UATCommand->SetOrder(GetUATCommandOrder());
+		BroadcastPropertyChanged();
+	}
+	void FUATCommandLaunchExtensionInstance::OnRemoved()
 	{
-		auto ToggleUATCommandEnabled = [this]()
+		ILauncherProfileUATCommandPtr ExistingUATCommand = GetUATCommand();
+		if (ExistingUATCommand.IsValid())
 		{
-			ILauncherProfileUATCommandPtr ExistingUATCommand = GetUATCommand();
-			if (ExistingUATCommand.IsValid())
-			{
-				OnUATCommandRemoved(ExistingUATCommand.ToSharedRef());
-				GetProfile()->RemoveUATCommand(*GetUATCommandInternalName());
-			}
-			else
-			{
-				ILauncherProfileUATCommandRef UATCommand = GetProfile()->FindOrAddUATCommand(*GetUATCommandInternalName());
-				OnUATCommandAdded(UATCommand);
-				UATCommand->SetOrder(GetUATCommandOrder());
-			}
-
-			BroadcastPropertyChanged();
-		};
-
-		MenuBuilder.AddMenuEntry(
-			GetExtension()->GetDisplayName(),
-			FText::GetEmpty(),
-			FSlateIcon(),
-			FUIAction(
-				FExecuteAction::CreateLambda(ToggleUATCommandEnabled),
-				FCanExecuteAction(),
-				FIsActionChecked::CreateRaw(this, &FUATCommandLaunchExtensionInstance::IsUATCommandActive )
-			),
-			NAME_None,
-			EUserInterfaceActionType::Check);
+			OnUATCommandRemoved(ExistingUATCommand.ToSharedRef());
+			GetProfile()->RemoveUATCommand(*GetUATCommandInternalName());
+		}
+		BroadcastPropertyChanged();
 	}
 
 	bool FUATCommandLaunchExtensionInstance::IsUATCommandActive() const
@@ -693,11 +818,6 @@ namespace ProjectLauncher
 		return GetProfile()->GetUATCommand(*GetUATCommandInternalName());
 	}
 
-	void FUATCommandLaunchExtensionInstance::GetExtensionsMenuEntry( FExtensionsMenuEntry& MenuEntry ) const
-	{
-		MenuEntry = FExtensionsMenuEntry::UATCommands;
-	}
-
 
 	bool FUATCommandLaunchExtensionInstance::CustomizeUATCommandLine( const ILauncherProfileUATCommandRef& InUATCommand, FString& InOutCommandLine )
 	{
@@ -736,40 +856,29 @@ namespace ProjectLauncher
 		}
 	}
 
-	void FAutomatedTestLaunchExtensionInstance::MakeCustomExtensionSubmenu( FMenuBuilder& MenuBuilder )
+
+	void FAutomatedTestLaunchExtensionInstance::OnAdded()
 	{
-		auto ToggleTestEnabled = [this]()
-		{
-			ILauncherProfileAutomatedTestPtr ExistingAutomatedTest = GetTest();
-			if (ExistingAutomatedTest.IsValid())
-			{
-				OnTestRemoved(ExistingAutomatedTest.ToSharedRef());
-				GetProfile()->RemoveAutomatedTest(*GetTestInternalName());
-			}
-			else
-			{
-				ILauncherProfileAutomatedTestRef AutomatedTest = GetProfile()->FindOrAddAutomatedTest(*GetTestInternalName());
-				OnTestAdded(AutomatedTest);
+		ILauncherProfileAutomatedTestRef AutomatedTest = GetProfile()->FindOrAddAutomatedTest(*GetTestInternalName());
+		OnTestAdded(AutomatedTest);
 
-				int32 Order = GetTestOrder();
-				check(Order > 0);
-				AutomatedTest->SetOrder(Order);
-			}
+		int32 Order = GetTestOrder();
+		check(Order > 0);
+		AutomatedTest->SetOrder(Order);
 
-			BroadcastPropertyChanged();
-		};
+		BroadcastPropertyChanged();
+
+	}
 
-		MenuBuilder.AddMenuEntry(
-			GetExtension()->GetDisplayName(),
-			FText::GetEmpty(),
-			FSlateIcon(),
-			FUIAction(
-				FExecuteAction::CreateLambda(ToggleTestEnabled),
-				FCanExecuteAction(),
-				FIsActionChecked::CreateRaw(this, &FAutomatedTestLaunchExtensionInstance::IsTestActive )
-			),
-			NAME_None,
-			EUserInterfaceActionType::Check);
+	void FAutomatedTestLaunchExtensionInstance::OnRemoved()
+	{
+		ILauncherProfileAutomatedTestPtr ExistingAutomatedTest = GetTest();
+		if (ExistingAutomatedTest.IsValid())
+		{
+			OnTestRemoved(ExistingAutomatedTest.ToSharedRef());
+			GetProfile()->RemoveAutomatedTest(*GetTestInternalName());
+		}
+		BroadcastPropertyChanged();
 	}
 
 	bool FAutomatedTestLaunchExtensionInstance::IsTestActive() const
@@ -782,11 +891,6 @@ namespace ProjectLauncher
 		return GetProfile()->GetAutomatedTest(*GetTestInternalName());
 	}
 
-	void FAutomatedTestLaunchExtensionInstance::GetExtensionsMenuEntry( FExtensionsMenuEntry& MenuEntry ) const
-	{
-		MenuEntry = FExtensionsMenuEntry::AutomatedTests;
-	}
-
 
 	bool FAutomatedTestLaunchExtensionInstance::CustomizeAutomatedTestCommandLine( const ILauncherProfileAutomatedTestRef& InAutomatedTest, FString& InOutCommandLine )
 	{
diff --git a/Engine/Plugins/Developer/ProjectLauncher/Source/ProjectLauncher/Private/ProfileTree/LaunchProfileTreeData.cpp b/Engine/Plugins/Developer/ProjectLauncher/Source/ProjectLauncher/Private/ProfileTree/LaunchProfileTreeData.cpp
--- a/Engine/Plugins/Developer/ProjectLauncher/Source/ProjectLauncher/Private/ProfileTree/LaunchProfileTreeData.cpp
+++ b/Engine/Plugins/Developer/ProjectLauncher/Source/ProjectLauncher/Private/ProfileTree/LaunchProfileTreeData.cpp
@@ -151,6 +151,7 @@ namespace ProjectLauncher
 		TreeNode->SortOrder = InSortOrder;
 		TreeNode->Depth = 0;
 		Nodes.Add(TreeNode);
+		RequestExpandHeadings.Add(TreeNode);
 
 		HeadingNodes.Add(InName, TreeNode);
 
@@ -182,6 +183,11 @@ namespace ProjectLauncher
 		Profile->RefreshCustomWarningsAndErrors();
 	}
 
+	void FLaunchProfileTreeData::RequestFullTreeRebuild()
+	{
+		bRequestTreeFullRebuild = true;
+	}
+
 	void FLaunchProfileTreeData::OnPropertyChanged()
 	{
 		for (TSharedPtr<FLaunchExtensionInstance> ExtensionInstance : ExtensionInstances)
@@ -219,6 +225,7 @@ namespace ProjectLauncher
 		TreeNode->SortOrder = InSortOrder;
 		TreeNode->Depth = Depth+1;
 		Children.Add(TreeNode);
+		TreeData->RequestExpandHeadings.Add(TreeNode);
 
 		SubHeadingNodes.Add(InName, TreeNode);
 
diff --git a/Engine/Plugins/Developer/ProjectLauncher/Source/ProjectLauncher/Private/Widgets/Profiles/SCustomLaunchCustomProfileEditor.cpp b/Engine/Plugins/Developer/ProjectLauncher/Source/ProjectLauncher/Private/Widgets/Profiles/SCustomLaunchCustomProfileEditor.cpp
--- a/Engine/Plugins/Developer/ProjectLauncher/Source/ProjectLauncher/Private/Widgets/Profiles/SCustomLaunchCustomProfileEditor.cpp
+++ b/Engine/Plugins/Developer/ProjectLauncher/Source/ProjectLauncher/Private/Widgets/Profiles/SCustomLaunchCustomProfileEditor.cpp
@@ -568,6 +568,64 @@ void SCustomLaunchCustomProfileEditor::SetProfile( const ILauncherProfilePtr& Pr
 	TreeView->RequestTreeRefresh();
 }
 
+void SCustomLaunchCustomProfileEditor::RebuildTree()
+{
+	if (!CurrentProfile.IsValid() || !TreeView.IsValid())
+	{
+		return;
+	}
+
+	// cache which headings were collapsed before we rebuild the tree so we can restore the state in the new tree
+	TSet<FString> CurrentlyCollapsed;
+	for ( TTuple<FString,ProjectLauncher::FLaunchProfileTreeNodePtr> Pair : TreeBuilder->GetProfileTree()->HeadingNodes)
+	{
+		CacheCollapsedRecursive(CurrentlyCollapsed, Pair.Key, Pair.Value);
+	}
+
+	TreeView->ClearRootItemsSource();
+	TreeBuilder.Reset();
+
+	TreeBuilder = ProjectLauncher::CreateTreeBuilder( CurrentProfile, Model.ToSharedRef() );
+	TreeView->SetTreeItemsSource( &TreeBuilder->GetProfileTree()->Nodes );
+
+	// restore the collapsed state
+	for ( TTuple<FString,ProjectLauncher::FLaunchProfileTreeNodePtr> Pair : TreeBuilder->GetProfileTree()->HeadingNodes)
+	{
+		ApplyCollapsedRecusive(CurrentlyCollapsed, Pair.Key, Pair.Value);
+	}
+
+	TreeView->RequestTreeRefresh();
+}
+
+void SCustomLaunchCustomProfileEditor::CacheCollapsedRecursive( TSet<FString>& Collapsed, const FString& NodeName, ProjectLauncher::FLaunchProfileTreeNodePtr Node) const
+{
+	if (TreeView->IsItemExpanded(Node))
+	{
+		for ( TTuple<FString, ProjectLauncher::FLaunchProfileTreeNodePtr> Pair : Node->SubHeadingNodes)
+		{
+			CacheCollapsedRecursive(Collapsed, NodeName + TEXT(":") + Pair.Key, Pair.Value);
+		}
+	}
+	else
+	{
+		Collapsed.Add(NodeName);
+	}
+}
+
+void SCustomLaunchCustomProfileEditor::ApplyCollapsedRecusive( TSet<FString>& Collapsed, const FString& NodeName, ProjectLauncher::FLaunchProfileTreeNodePtr Node) const
+{
+	if (!Collapsed.Contains(NodeName))
+	{
+		TreeView->SetItemExpansion(Node, true);
+
+		for ( TTuple<FString, ProjectLauncher::FLaunchProfileTreeNodePtr> Pair : Node->SubHeadingNodes)
+		{
+			ApplyCollapsedRecusive(Collapsed,  NodeName + TEXT(":") + Pair.Key, Pair.Value);
+		}
+	}
+}
+
+
 TSharedRef<SWidget> SCustomLaunchCustomProfileEditor::MakeExtensionsMenu()
 {
 	ProjectLauncher::FLaunchProfileTreeDataRef TreeData = TreeBuilder->GetProfileTree();
@@ -578,23 +636,24 @@ TSharedRef<SWidget> SCustomLaunchCustomProfileEditor::MakeExtensionsMenu()
 	{
 		FText MenuEntryName;
 		FText DisplayName;
-		TArray<TSharedPtr<ProjectLauncher::FLaunchExtensionInstance>> ExtensionInstances;
+		TArray<TSharedRef<ProjectLauncher::FLaunchExtension>> Extensions;
 	};
 
 	struct FExtensionMenuSection
 	{
 		FText DisplayName;
 		TMap<FString,TSharedPtr<FExtensionMenuSubMenu>> SubMenus;
-		TArray<TSharedPtr<ProjectLauncher::FLaunchExtensionInstance>> ExtensionInstances;
+		TArray<TSharedRef<ProjectLauncher::FLaunchExtension>> Extensions;
 	};
 
 	TMap<FString,FExtensionMenuSection> Sections;
 
-	for (TSharedPtr<ProjectLauncher::FLaunchExtensionInstance> ExtensionInstance : TreeData->ExtensionInstances)
+	TArray<TSharedRef<ProjectLauncher::FLaunchExtension>> Extensions = ProjectLauncher::FLaunchExtension::GetExtensions();
+	for (const TSharedRef<ProjectLauncher::FLaunchExtension>& Extension : Extensions)
 	{
-		ProjectLauncher::FLaunchExtensionInstance::FExtensionsMenuEntry ExtensionMenuEntry;
-		ExtensionInstance->GetExtensionsMenuEntry(ExtensionMenuEntry);
-		if (ExtensionMenuEntry.Type == ProjectLauncher::FLaunchExtensionInstance::FExtensionsMenuEntry::Type_None)
+		ProjectLauncher::FLaunchExtension::FExtensionsMenuEntry ExtensionMenuEntry;
+		Extension->GetExtensionsMenuEntry(ExtensionMenuEntry);
+		if (ExtensionMenuEntry.Type == ProjectLauncher::FLaunchExtension::FExtensionsMenuEntry::Type_None)
 		{
 			continue;
 		}
@@ -604,8 +663,8 @@ TSharedRef<SWidget> SCustomLaunchCustomProfileEditor::MakeExtensionsMenu()
 		FText SectionDisplayName = ExtensionMenuEntry.SectionDisplayName;
 		if (SectionName.IsEmpty())
 		{
-			SectionName = ExtensionInstance->GetExtension()->GetInternalName();
-			SectionDisplayName = ExtensionInstance->GetExtension()->GetDisplayName();
+			SectionName = Extension->GetInternalName();
+			SectionDisplayName = Extension->GetDisplayName();
 		}
 		FExtensionMenuSection* Section = Sections.Find(SectionName);
 		if (Section == nullptr)
@@ -614,7 +673,7 @@ TSharedRef<SWidget> SCustomLaunchCustomProfileEditor::MakeExtensionsMenu()
 			Section->DisplayName = SectionDisplayName;
 		}
 
-		if (ExtensionMenuEntry.Type == ProjectLauncher::FLaunchExtensionInstance::FExtensionsMenuEntry::Type_SubMenu)
+		if (ExtensionMenuEntry.Type == ProjectLauncher::FLaunchExtension::FExtensionsMenuEntry::Type_SubMenu)
 		{
 			// get submenu
 			TSharedPtr<FExtensionMenuSubMenu>* SubMenu = Section->SubMenus.Find(ExtensionMenuEntry.SubmenuName);
@@ -624,11 +683,11 @@ TSharedRef<SWidget> SCustomLaunchCustomProfileEditor::MakeExtensionsMenu()
 				(*SubMenu)->DisplayName = ExtensionMenuEntry.SubmenuDisplayName;
 			}
 
-			(*SubMenu)->ExtensionInstances.Add(ExtensionInstance);
+			(*SubMenu)->Extensions.Add(Extension);
 		}
 		else
 		{
-			Section->ExtensionInstances.Add(ExtensionInstance);
+			Section->Extensions.Add(Extension);
 		}
 	}
 
@@ -642,31 +701,106 @@ TSharedRef<SWidget> SCustomLaunchCustomProfileEditor::MakeExtensionsMenu()
 
 	// ...construct the extensions menu...
 
-	const bool bShouldCloseWindowAfterMenuSelection = false;
+	const bool bShouldCloseWindowAfterMenuSelection = true;
 	const bool bCloseSelfOnly = false;
 	const bool bSearchable = false;
 	FMenuBuilder MenuBuilder(bShouldCloseWindowAfterMenuSelection, nullptr, nullptr, bCloseSelfOnly, &FCoreStyle::Get(), bSearchable );
 	
+
+	auto MakeExtensionMenu = [this](FMenuBuilder& MenuBuilder, TSharedRef<ProjectLauncher::FLaunchExtension> Extension)
+	{
+		TSharedPtr<ProjectLauncher::FLaunchExtensionInstance> ExtensionInstance = ProjectLauncher::FLaunchExtension::GetExtensionInstance(Extension, CurrentProfile.ToSharedRef());
+
+		if (!Extension->IsAlwaysCreated(CurrentProfile.ToSharedRef(), Model.ToSharedRef()))
+		{
+			auto OnToggleExtension = [this, Extension, ExtensionInstance]()
+			{
+				if (ExtensionInstance.IsValid())
+				{
+					// defer the delete until the next tick
+					PendingExtensionDeletes.AddUnique(Extension);
+				}
+				else
+				{
+					TSharedPtr<ProjectLauncher::FLaunchExtensionInstance> Instance = ProjectLauncher::FLaunchExtension::CreateExtensionInstance(Extension, CurrentProfile.ToSharedRef(), Model.ToSharedRef(), TreeBuilder->GetProfileTree().Get());
+					if (Instance.IsValid())
+					{
+						Instance->CustomizeTree(TreeBuilder->GetProfileTree().Get());
+						TreeBuilder->OnPropertyChanged();
+						CurrentProfile->RefreshCustomWarningsAndErrors();
+					}
+				}
+				
+			};
+
+			auto CanToggleExtension = [this, Extension, ExtensionInstance]
+			{
+				if (ExtensionInstance.IsValid())
+				{
+					return ExtensionInstance->CanBeRemoved();
+				}
+				else
+				{
+					return Extension->CanBeCreated(CurrentProfile.ToSharedRef(), Model.ToSharedRef());
+				}
+			};
+
+			auto GetToolTipText = [this, Extension, ExtensionInstance]
+			{
+				if (ExtensionInstance.IsValid() && !ExtensionInstance->CanBeRemoved())
+				{
+					return LOCTEXT("CannotRemoveExtensionTooltip", "The extension cannot be removed at this time");
+				}
+				else if (!ExtensionInstance.IsValid() && !Extension->CanBeCreated(CurrentProfile.ToSharedRef(), Model.ToSharedRef()))
+				{
+					return LOCTEXT("CannotAddExtensionTooltip", "The extension cannot be added at this time");
+				}
+
+				return FText::GetEmpty();
+			};
+
+			MenuBuilder.AddMenuEntry(
+				Extension->GetDisplayName(),
+				GetToolTipText(),
+				FSlateIcon(),
+				FUIAction(
+					FExecuteAction::CreateLambda(OnToggleExtension),
+					FCanExecuteAction::CreateLambda(CanToggleExtension),
+					FIsActionChecked::CreateLambda( [ExtensionInstance]() { return ExtensionInstance.IsValid(); } )
+				),
+				NAME_None,
+				EUserInterfaceActionType::Check);
+		}
+
+		if (ExtensionInstance)
+		{
+			ExtensionInstance->MakeCustomExtensionSubmenu(MenuBuilder);
+		}
+
+		Extension->MakeCustomExtensionSubmenu(MenuBuilder, CurrentProfile.ToSharedRef(), Model.ToSharedRef());
+	};
+
+
 	for (const auto& SectionItr : Sections)
 	{
 		// start the section
 		MenuBuilder.BeginSection(NAME_None, SectionItr.Value.DisplayName);
 
 		// add all extensions
-		for (TSharedPtr<ProjectLauncher::FLaunchExtensionInstance> ExtensionInstance : SectionItr.Value.ExtensionInstances)
+		for (TSharedRef<ProjectLauncher::FLaunchExtension> Extension : SectionItr.Value.Extensions)
 		{
-			ExtensionInstance->MakeCustomExtensionSubmenu(MenuBuilder);
+			MakeExtensionMenu(MenuBuilder, Extension);
 		}
 
 		// add all submenus
 		for (const auto& SubMenuItr : SectionItr.Value.SubMenus)
 		{
-			auto MakeExtensionSubMenu = [](FMenuBuilder& MenuBuilder, TSharedPtr<FExtensionMenuSubMenu> SubMenu)
+			auto MakeExtensionSubMenu = [MakeExtensionMenu](FMenuBuilder& MenuBuilder, TSharedPtr<FExtensionMenuSubMenu> SubMenu)
 			{
 				MenuBuilder.SetSearchable(bSearchable);
-				for (TSharedPtr<ProjectLauncher::FLaunchExtensionInstance> ExtensionInstance : SubMenu->ExtensionInstances)
+				for (TSharedRef<ProjectLauncher::FLaunchExtension> Extension : SubMenu->Extensions)
 				{
-					ExtensionInstance->MakeCustomExtensionSubmenu(MenuBuilder);
+					MakeExtensionMenu(MenuBuilder, Extension);
 				}
 			};
 
@@ -693,10 +827,47 @@ void SCustomLaunchCustomProfileEditor::Tick(const FGeometry& AllottedGeometry, c
 {
 	SCompoundWidget::Tick(AllottedGeometry, InCurrentTime, InDeltaTime);
 
-	if (TreeBuilder.IsValid() && TreeBuilder->GetProfileTree()->bRequestTreeRefresh)
+	if (TreeBuilder.IsValid())
 	{
-		TreeView->RequestTreeRefresh();
-		TreeBuilder->GetProfileTree()->bRequestTreeRefresh = false;
+		// expand anything that has just been added
+		if (!TreeBuilder->GetProfileTree()->RequestExpandHeadings.IsEmpty())
+		{
+			for (const ProjectLauncher::FLaunchProfileTreeNodePtr& Node : TreeBuilder->GetProfileTree()->RequestExpandHeadings)
+			{
+				TreeView->SetItemExpansion(Node, true);
+			}
+			TreeBuilder->GetProfileTree()->RequestExpandHeadings.Reset();
+		}
+
+		// check if there's any refreshing / rebuilding to be done
+		if (TreeBuilder->GetProfileTree()->bRequestTreeFullRebuild || PendingExtensionDeletes.Num() > 0)
+		{
+			TreeBuilder->GetProfileTree()->bRequestTreeFullRebuild = false;
+
+			// wait until after the slate tick before applying deletes & refreshing
+			ExecuteOnGameThread(UE_SOURCE_LOCATION,[this, PendingDeletes = MoveTemp(PendingExtensionDeletes)]()
+			{
+				bool bProfileDirty = false;
+				for (const TSharedRef<ProjectLauncher::FLaunchExtension>& Extension : PendingDeletes)
+				{
+					ProjectLauncher::FLaunchExtension::DestroyExtensionInstance(Extension, CurrentProfile.ToSharedRef(), Model.ToSharedRef());
+					bProfileDirty = true;
+				}
+
+				RebuildTree();
+
+				if (bProfileDirty)
+				{
+					TreeBuilder->OnPropertyChanged();
+					CurrentProfile->RefreshCustomWarningsAndErrors();
+				}
+			});
+		}
+		else if (TreeBuilder->GetProfileTree()->bRequestTreeRefresh)
+		{
+			TreeView->RequestTreeRefresh();
+			TreeBuilder->GetProfileTree()->bRequestTreeRefresh = false;
+		}
 	}
 }
 
diff --git a/Engine/Plugins/Developer/ProjectLauncher/Source/ProjectLauncher/Private/Widgets/Profiles/SCustomLaunchCustomProfileEditor.h b/Engine/Plugins/Developer/ProjectLauncher/Source/ProjectLauncher/Private/Widgets/Profiles/SCustomLaunchCustomProfileEditor.h
--- a/Engine/Plugins/Developer/ProjectLauncher/Source/ProjectLauncher/Private/Widgets/Profiles/SCustomLaunchCustomProfileEditor.h
+++ b/Engine/Plugins/Developer/ProjectLauncher/Source/ProjectLauncher/Private/Widgets/Profiles/SCustomLaunchCustomProfileEditor.h
@@ -16,6 +16,11 @@ class ITableRow;
 class STableViewBase;
 template<typename T> class STreeView;
 
+namespace ProjectLauncher
+{
+	class FLaunchExtension;
+	class FLaunchExtensionInstance;
+};
 typedef STreeView<ProjectLauncher::FLaunchProfileTreeNodePtr> SLaunchProfileTreeView;
 typedef STableRow<ProjectLauncher::FLaunchProfileTreeNodePtr> SLaunchProfileTreeRow;
 
@@ -34,6 +39,8 @@ public:
 
 	TSharedRef<SWidget> MakeExtensionsMenu();
 
+	void RebuildTree();
+
 	// IScrollableWidget interface
 	virtual FVector2D GetScrollDistance() override;
 	virtual FVector2D GetScrollDistanceRemaining() override;
@@ -45,10 +52,14 @@ private:
 	TSharedRef<ITableRow> OnGenerateWidgetForTreeNode( ProjectLauncher::FLaunchProfileTreeNodePtr TreeNode, const TSharedRef<STableViewBase>& OwnerTable );
 	void OnGetChildren(ProjectLauncher::FLaunchProfileTreeNodePtr TreeNode, TArray<ProjectLauncher::FLaunchProfileTreeNodePtr>& OutChildren);
 
+	void CacheCollapsedRecursive( TSet<FString>& Collapsed, const FString& NodeName, ProjectLauncher::FLaunchProfileTreeNodePtr Node) const;
+	void ApplyCollapsedRecusive( TSet<FString>& Collapsed, const FString& NodeName, ProjectLauncher::FLaunchProfileTreeNodePtr Node) const;
+
 	ILauncherProfilePtr CurrentProfile;
 	TSharedPtr<ProjectLauncher::FModel> Model;
 	TSharedPtr<ProjectLauncher::ILaunchProfileTreeBuilder> TreeBuilder;
 	TSharedPtr<SLaunchProfileTreeView> TreeView;
+	TArray<TSharedRef<ProjectLauncher::FLaunchExtension>> PendingExtensionDeletes;
 
 	float SplitterPos = 0.6f;
 	
diff --git a/Engine/Plugins/Developer/ProjectLauncher/Source/ProjectLauncher/Public/Extension/LaunchExtension.h b/Engine/Plugins/Developer/ProjectLauncher/Source/ProjectLauncher/Public/Extension/LaunchExtension.h
--- a/Engine/Plugins/Developer/ProjectLauncher/Source/ProjectLauncher/Public/Extension/LaunchExtension.h
+++ b/Engine/Plugins/Developer/ProjectLauncher/Source/ProjectLauncher/Public/Extension/LaunchExtension.h
@@ -109,37 +109,5 @@
 		virtual void OnValidateProfile() {};
 
 		/**
-		 * Describes how this extension is represented in the extensions menu
-		 */
-		struct FExtensionsMenuEntry
-		{
-			// common menu locations
-			static UE_API const FExtensionsMenuEntry None;
-			static UE_API const FExtensionsMenuEntry Default;
-			static UE_API const FExtensionsMenuEntry OwnSection;
-			static UE_API const FExtensionsMenuEntry Deprecated;
-			static UE_API const FExtensionsMenuEntry AutomatedTests;
-			static UE_API const FExtensionsMenuEntry UATCommands;
-
-			// internal names
-			FString SectionName; // if this is empty, the extension will be placed in its own section
-			FString SubmenuName;
-
-			// friendly names
-			FText SectionDisplayName;
-			FText SubmenuDisplayName;
-
-			// type of menu entry
-			enum EType
-			{
-				Type_None,			// does not appear in the extensions menu
-				Type_Section,		// appears in the extensions menu in the common section.
-				Type_SubMenu,		// appears in the extensions menu in a submenu of the SectionName section
-			};
-			EType Type = Type_None;
-		};
-		UE_API virtual void GetExtensionsMenuEntry( FExtensionsMenuEntry& MenuEntry ) const;
-
-		/**
 		 * Create the default heading for this extension
 		 */
@@ -156,11 +124,31 @@
 		 */
 		virtual void MakeCustomExtensionSubmenu(FMenuBuilder& MenuBuilder) {}
 
+		/*
+		 * Whether the given extension can be removed or not
+		 */
+		virtual bool CanBeRemoved() const { return true; }
+
 		/** 
 		 * Get the extension that instantiated us
 		 */
 		inline TSharedRef<FLaunchExtension> GetExtension() const { return Extension; }
 
+		/** 
+		 * Get the profile we were instantiated for
+		 */
+		inline ILauncherProfileRef GetProfile() const { return Profile; }
+
+		/** 
+		 * Get the reference default profile
+		 */
+		inline ILauncherProfileRef GetDefaultProfile() const { return DefaultProfile; }
+
+		/** 
+		 * Get the model, for general purpose helper functions
+		 */
+		inline TSharedRef<FModel> GetModel() const { return Model; }
+
 		/**
 		 * Get the property changed delegate that will be called when an extension changes the profile
 		 */
@@ -171,6 +159,9 @@
 		 */
 		UE_API virtual void InternalInitialize() {};
 
+		UE_API virtual void OnAdded() {};
+		UE_API virtual void OnRemoved() {};
+
 	protected:
 
 		UE_DEPRECATED(5.8, "Use ProjectLauncher::CmdLineUtils version")
@@ -210,21 +201,6 @@
 		 */
 		UE_API void SetCommandLine( const FString& CommandLine );
 
-		/** 
-		 * Get the profile we were instantiated for
-		 */
-		inline ILauncherProfileRef GetProfile() const { return Profile; }
-
-		/** 
-		 * Get the reference default profile
-		 */
-		inline ILauncherProfileRef GetDefaultProfile() const { return DefaultProfile; }
-
-		/** 
-		 * Get the model, for general purpose helper functions
-		 */
-		inline TSharedRef<FModel> GetModel() const { return Model; }
-
 		/**
 		 * Get the name of the platform to launch on. this is the highest priority platform/target in the list of platforms
 		 */
@@ -366,25 +342,18 @@
 	public:
 		UE_API FUATCommandLaunchExtensionInstance( FArgs& InArgs );
 
-		/*
-		 * Populate the custom extension menu for this UAT command launch extension. This will add a menu item to toggle the UAT command on and off
-		 */
-		UE_API virtual void MakeCustomExtensionSubmenu( FMenuBuilder& MenuBuilder ) override;
-
 		/**
-		 * Describes how this extension is represented in the extensions menu - a common section for all UAT commands
-		 */
-		UE_API virtual void GetExtensionsMenuEntry( FExtensionsMenuEntry& MenuEntry ) const override;
-		
-		/**
 		 * Determine if this UAT Command is currently active
 		 */
 		UE_API bool IsUATCommandActive() const;
 
 	protected:
+		UE_API virtual void OnAdded() override;
+		UE_API virtual void OnRemoved() override;
+
+
 		UE_API bool CustomizeUATCommandLine( const ILauncherProfileUATCommandRef& InUATCommand, FString& InOutCommandLine ) override;
-		virtual bool HasCustomExtensionMenu() const { return true; }
-
+		
 		/**
 		 * Advanced hook to allow any advanced modification of the command line when this extension's UAT Command is launched
 		 */
@@ -437,22 +406,15 @@
 	public:
 		UE_API FAutomatedTestLaunchExtensionInstance( FArgs& InArgs );
 
-		/*
-		 * Populate the custom extension menu for this automated test extension. This will add a menu item to toggle the test on and off
-		 */
-		UE_API virtual void MakeCustomExtensionSubmenu( FMenuBuilder& MenuBuilder ) override;
-
 		/**
-		 * Describes how this extension is represented in the extensions menu - a common section for all automated tests
-		 */
-		UE_API virtual void GetExtensionsMenuEntry( FExtensionsMenuEntry& MenuEntry ) const override;
-		
-		/**
 		 * Determine if this automated test is currently active
 		 */
 		UE_API bool IsTestActive() const;
 
 	protected:
+		UE_API virtual void OnAdded() override;
+		UE_API virtual void OnRemoved() override;
+
 		UE_API bool CustomizeAutomatedTestCommandLine( const ILauncherProfileAutomatedTestRef& InAutomatedTest, FString& InOutCommandLine ) override;
 		virtual bool HasCustomExtensionMenu() const { return true; }
 
@@ -523,6 +485,20 @@
 		 */
 		virtual FText GetDisplayName() const = 0;
 
+		/**
+		 * Whether this extension is always created. the user cannot remove it
+		 */
+		virtual bool IsAlwaysCreated( ILauncherProfileRef InProfile, TSharedRef<FModel> InModel ) const { return false; }
 
 		/**
+		 * Whether this extension will be created by default. the user can still remove it
+		 */
+		virtual bool IsCreatedByDefault( ILauncherProfileRef InProfile, TSharedRef<FModel> InModel ) const { return false; }
+
+		/**
+		 * Whether this extension can be created at this time
+		 */
+		virtual bool CanBeCreated( ILauncherProfileRef InProfile, TSharedRef<FModel> InModel ) const { return true; }
+
+		/**
 		 * Instantiate all compatible extensions
@@ -534,6 +510,47 @@
 		 */
 		UE_API static TArray<TSharedPtr<FLaunchExtensionInstance>> CreateExtensionInstancesForProfile( ILauncherProfileRef InProfile, TSharedRef<FModel> InModel, FLaunchProfileTreeData& InOwnerTreeData );
 
+		UE_API static TSharedPtr<FLaunchExtensionInstance> CreateExtensionInstance( TSharedRef<FLaunchExtension> InExtension, ILauncherProfileRef InProfile, TSharedRef<FModel> InModel, FLaunchProfileTreeData& InOwnerTreeData );
+		UE_API static void DestroyExtensionInstance(TSharedRef<FLaunchExtension> InExtension, ILauncherProfileRef InProfile, TSharedRef<FModel> InModel);
+		UE_API static TSharedPtr<FLaunchExtensionInstance> GetExtensionInstance(TSharedRef<FLaunchExtension> InExtension, ILauncherProfileRef InProfile);
 
+		UE_API static TArray<TSharedRef<FLaunchExtension>> GetExtensions();
+
+
+		/**
+		 * Describes how this extension is represented in the extensions menu
+		 */
+		struct FExtensionsMenuEntry
+		{
+			// common menu locations
+			static UE_API const FExtensionsMenuEntry None;
+			static UE_API const FExtensionsMenuEntry Default;
+			static UE_API const FExtensionsMenuEntry OwnSection;
+			static UE_API const FExtensionsMenuEntry Deprecated;
+			static UE_API const FExtensionsMenuEntry AutomatedTests;
+			static UE_API const FExtensionsMenuEntry UATCommands;
+
+			// internal names
+			FString SectionName; // if this is empty, the extension will be placed in its own section
+			FString SubmenuName;
+
+			// friendly names
+			FText SectionDisplayName;
+			FText SubmenuDisplayName;
+
+			// type of menu entry
+			enum EType
+			{
+				Type_None,			// does not appear in the extensions menu
+				Type_Section,		// appears in the extensions menu in the common section.
+				Type_SubMenu,		// appears in the extensions menu in a submenu of the SectionName section
+			};
+			EType Type = Type_None;
+		};
+		UE_API virtual void GetExtensionsMenuEntry( FExtensionsMenuEntry& MenuEntry ) const;
+
+		virtual void MakeCustomExtensionSubmenu(FMenuBuilder& MenuBuilder, ILauncherProfileRef InProfile, TSharedRef<FModel> InModel) {}
+
+
 	private:
 		/**
@@ -545,6 +562,21 @@
 	};
 
 
+
+	class FAutomatedTestLaunchExtension : public FLaunchExtension
+	{
+	public:
+		UE_API virtual void GetExtensionsMenuEntry( FExtensionsMenuEntry& MenuEntry ) const override;
+
+	};
+
+	class FUATCommandLaunchExtension : public FLaunchExtension
+	{
+	public:
+		UE_API virtual void GetExtensionsMenuEntry( FExtensionsMenuEntry& MenuEntry ) const override;
+
+	};
+
 };
 
 #undef UE_API
diff --git a/Engine/Plugins/Developer/ProjectLauncher/Source/ProjectLauncher/Public/ProfileTree/LaunchProfileTreeData.h b/Engine/Plugins/Developer/ProjectLauncher/Source/ProjectLauncher/Public/ProfileTree/LaunchProfileTreeData.h
--- a/Engine/Plugins/Developer/ProjectLauncher/Source/ProjectLauncher/Public/ProfileTree/LaunchProfileTreeData.h
+++ b/Engine/Plugins/Developer/ProjectLauncher/Source/ProjectLauncher/Public/ProfileTree/LaunchProfileTreeData.h
@@ -189,6 +189,7 @@ namespace ProjectLauncher
 		/* Notify of property changes and refresh requests. */
 		UE_API void OnPropertyChanged();
 		UE_API void RequestTreeRefresh();
+		UE_API void RequestFullTreeRebuild();
 
 		/* Member variables */
 		ILauncherProfilePtr Profile;                                     // associated profile
@@ -197,7 +198,9 @@ namespace ProjectLauncher
 		TMap<FString,FLaunchProfileTreeNodePtr> HeadingNodes;            // named top-level headings
 		ILaunchProfileTreeBuilder* TreeBuilder;                          // UI builder interface
 		TArray<TSharedPtr<FLaunchExtensionInstance>> ExtensionInstances; // extension instances attached
+		TSet<FLaunchProfileTreeNodePtr> RequestExpandHeadings;           // headings to expand
 		bool bRequestTreeRefresh = false;                                // request to refresh tree UI
+		bool bRequestTreeFullRebuild = false;                            // request to fully recreate the profile tree
 	};
 
 
diff --git a/Engine/Plugins/Performance/AutomatedPerfTesting/Source/AutomatedPerfTesting/Private/LaunchExtension/AutomatedPerfTestLaunchExtension.cpp b/Engine/Plugins/Performance/AutomatedPerfTesting/Source/AutomatedPerfTesting/Private/LaunchExtension/AutomatedPerfTestLaunchExtension.cpp
--- a/Engine/Plugins/Performance/AutomatedPerfTesting/Source/AutomatedPerfTesting/Private/LaunchExtension/AutomatedPerfTestLaunchExtension.cpp
+++ b/Engine/Plugins/Performance/AutomatedPerfTesting/Source/AutomatedPerfTesting/Private/LaunchExtension/AutomatedPerfTestLaunchExtension.cpp
@@ -76,38 +76,6 @@ namespace AutomatedPerfTestConfig
 
 using namespace AutomatedPerfTestConfig;
 
-void FAutomatedPerfTestLaunchExtensionInstance::MakeCustomExtensionSubmenu(FMenuBuilder& MenuBuilder)
-{
-	const auto ToggleTestEnabled = [this]()
-	{
-		ILauncherProfileAutomatedTestPtr ExistingAutomatedTest = GetTest();
-		if (ExistingAutomatedTest.IsValid())
-		{
-			OnTestRemoved(ExistingAutomatedTest.ToSharedRef());
-			GetProfile()->RemoveAutomatedTest(*GetTestInternalName());
-		}
-		else
-		{
-			ILauncherProfileAutomatedTestRef AutomatedTest = GetProfile()->FindOrAddAutomatedTest(*GetTestInternalName());
-			OnTestAdded(AutomatedTest);
-		}
-
-		BroadcastPropertyChanged();
-	};
-
-	MenuBuilder.AddMenuEntry(
-		LOCTEXT("EnableAutomatedPerfTestLabel", "Automated Perf Test"),
-		FText::GetEmpty(),
-		FSlateIcon(),
-		FUIAction(
-			FExecuteAction::CreateLambda(ToggleTestEnabled),
-			FCanExecuteAction(),
-			FIsActionChecked::CreateRaw(this, &FAutomatedTestLaunchExtensionInstance::IsTestActive)
-		),
-		NAME_None,
-		EUserInterfaceActionType::Check);
-}
-
 void FAutomatedPerfTestLaunchExtensionInstance::OnPropertyChanged()
 {
 	ILauncherProfileAutomatedTestPtr Test = GetTest();
diff --git a/Engine/Plugins/Performance/AutomatedPerfTesting/Source/AutomatedPerfTesting/Private/LaunchExtension/AutomatedPerfTestLaunchExtension.h b/Engine/Plugins/Performance/AutomatedPerfTesting/Source/AutomatedPerfTesting/Private/LaunchExtension/AutomatedPerfTestLaunchExtension.h
--- a/Engine/Plugins/Performance/AutomatedPerfTesting/Source/AutomatedPerfTesting/Private/LaunchExtension/AutomatedPerfTestLaunchExtension.h
+++ b/Engine/Plugins/Performance/AutomatedPerfTesting/Source/AutomatedPerfTesting/Private/LaunchExtension/AutomatedPerfTestLaunchExtension.h
@@ -24,9 +24,6 @@ public:
 	FAutomatedPerfTestLaunchExtensionInstance(FArgs& InArgs) : FAutomatedTestLaunchExtensionInstance(InArgs) {};
 	virtual ~FAutomatedPerfTestLaunchExtensionInstance() = default;
 
-	virtual bool HasCustomExtensionMenu() const { return true; }
-	virtual void MakeCustomExtensionSubmenu(FMenuBuilder& MenuBuilder) override;
-
 	virtual void OnPropertyChanged() override;
 	virtual void OnTestAdded(ILauncherProfileAutomatedTestRef AutomatedTest) override;
 	virtual const FString GetTestInternalName() const override;
@@ -72,7 +69,7 @@ private:
 };
 
 
-class FAutomatedPerfTestLaunchExtension : public ProjectLauncher::FLaunchExtension
+class FAutomatedPerfTestLaunchExtension : public ProjectLauncher::FAutomatedTestLaunchExtension
 {
 public:
 	virtual TSharedPtr<ProjectLauncher::FLaunchExtensionInstance> CreateInstanceForProfile(ProjectLauncher::FLaunchExtensionInstance::FArgs& InArgs) override;
