Splitting the launch extension implementations into separate source files as the code is getting too long.

#rnx
#Synced-CL 49557694

diff --git a/Engine/Plugins/Developer/ProjectLauncher/Source/CommonLaunchExtensions/Private/BootTest/BootTestLaunchExtension.h b/Engine/Plugins/Developer/ProjectLauncher/Source/CommonLaunchExtensions/Private/BootTest/BootTestLaunchExtension.h
--- a/Engine/Plugins/Developer/ProjectLauncher/Source/CommonLaunchExtensions/Private/BootTest/BootTestLaunchExtension.h
+++ b/Engine/Plugins/Developer/ProjectLauncher/Source/CommonLaunchExtensions/Private/BootTest/BootTestLaunchExtension.h
@@ -2,7 +2,7 @@
 
 #pragma once
 
-#include "Extension/LaunchExtension.h"
+#include "Extension/AutomatedTestLaunchExtension.h"
 
 
 
diff --git a/Engine/Plugins/Developer/ProjectLauncher/Source/ProjectLauncher/Private/Extension/AutomatedTestLaunchExtension.cpp b/Engine/Plugins/Developer/ProjectLauncher/Source/ProjectLauncher/Private/Extension/AutomatedTestLaunchExtension.cpp
new file mode 100644
--- /dev/null
+++ b/Engine/Plugins/Developer/ProjectLauncher/Source/ProjectLauncher/Private/Extension/AutomatedTestLaunchExtension.cpp
@@ -0,0 +1,83 @@
+// Copyright Epic Games, Inc. All Rights Reserved.
+
+#include "Extension/AutomatedTestLaunchExtension.h"
+
+#define LOCTEXT_NAMESPACE "FAutomatedTestLaunchExtensionInstance"
+
+
+namespace ProjectLauncher
+{
+	FAutomatedTestLaunchExtensionInstance::FAutomatedTestLaunchExtensionInstance( FArgs& InArgs )
+		: FLaunchExtensionInstance(InArgs)
+	{
+	}
+
+	void FAutomatedTestLaunchExtensionInstance::InternalInitialize()
+	{
+		ILauncherProfileAutomatedTestPtr AutomatedTest = GetTest();
+		if (AutomatedTest.IsValid())
+		{
+			OnTestAdded(AutomatedTest.ToSharedRef());
+
+			int32 Order = GetTestOrder();
+			check(Order > 0);
+			AutomatedTest->SetOrder(Order);
+		}
+	}
+
+
+	void FAutomatedTestLaunchExtensionInstance::OnAdded()
+	{
+		ILauncherProfileAutomatedTestRef AutomatedTest = GetProfile()->FindOrAddAutomatedTest(*GetTestInternalName());
+		OnTestAdded(AutomatedTest);
+
+		int32 Order = GetTestOrder();
+		check(Order > 0);
+		AutomatedTest->SetOrder(Order);
+
+		BroadcastPropertyChanged();
+
+	}
+
+	void FAutomatedTestLaunchExtensionInstance::OnRemoved()
+	{
+		ILauncherProfileAutomatedTestPtr ExistingAutomatedTest = GetTest();
+		if (ExistingAutomatedTest.IsValid())
+		{
+			OnTestRemoved(ExistingAutomatedTest.ToSharedRef());
+			GetProfile()->RemoveAutomatedTest(*GetTestInternalName());
+		}
+		BroadcastPropertyChanged();
+	}
+
+	bool FAutomatedTestLaunchExtensionInstance::IsTestActive() const
+	{
+		return GetTest() != nullptr;
+	}
+
+	ILauncherProfileAutomatedTestPtr FAutomatedTestLaunchExtensionInstance::GetTest() const
+	{
+		return GetProfile()->GetAutomatedTest(*GetTestInternalName());
+	}
+
+
+	bool FAutomatedTestLaunchExtensionInstance::CustomizeAutomatedTestCommandLine( const ILauncherProfileAutomatedTestRef& InAutomatedTest, FString& InOutCommandLine )
+	{
+		if (InAutomatedTest == GetTest())
+		{
+			CustomizeAutomatedTestCommandLine(InOutCommandLine);
+			return true;
+		}
+
+		return false;
+	}
+
+
+
+	void FAutomatedTestLaunchExtension::GetExtensionsMenuEntry( FExtensionsMenuEntry& MenuEntry ) const
+	{
+		MenuEntry = FExtensionsMenuEntry::AutomatedTests;
+	}
+}
+
+#undef LOCTEXT_NAMESPACE
diff --git a/Engine/Plugins/Developer/ProjectLauncher/Source/ProjectLauncher/Private/Extension/LaunchExtension.cpp b/Engine/Plugins/Developer/ProjectLauncher/Source/ProjectLauncher/Private/Extension/LaunchExtension.cpp
--- a/Engine/Plugins/Developer/ProjectLauncher/Source/ProjectLauncher/Private/Extension/LaunchExtension.cpp
+++ b/Engine/Plugins/Developer/ProjectLauncher/Source/ProjectLauncher/Private/Extension/LaunchExtension.cpp
@@ -409,15 +409,7 @@ namespace ProjectLauncher
 
 
 
-	void FAutomatedTestLaunchExtension::GetExtensionsMenuEntry( FExtensionsMenuEntry& MenuEntry ) const
-	{
-		MenuEntry = FExtensionsMenuEntry::AutomatedTests;
-	}
 
-	void FUATCommandLaunchExtension::GetExtensionsMenuEntry( FExtensionsMenuEntry& MenuEntry ) const
-	{
-		MenuEntry = FExtensionsMenuEntry::UATCommands;
-	}
 
 
 
@@ -771,137 +763,6 @@ namespace ProjectLauncher
 	{
 		return ProfileTreeData.AddHeading(Extension->GetInternalName(), OverrideDisplayName, GetHeadingSortOrder() );
 	}
-
-
-
-
-	FUATCommandLaunchExtensionInstance::FUATCommandLaunchExtensionInstance( FArgs& InArgs )
-		: FLaunchExtensionInstance(InArgs)
-	{
-	}
-
-	void FUATCommandLaunchExtensionInstance::InternalInitialize()
-	{
-		ILauncherProfileUATCommandPtr UATCommand = GetUATCommand();
-		if (UATCommand.IsValid())
-		{
-			OnUATCommandAdded(UATCommand.ToSharedRef());
-			UATCommand->SetOrder(GetUATCommandOrder());
-		}
-	}
-
-	void FUATCommandLaunchExtensionInstance::OnAdded()
-	{
-		ILauncherProfileUATCommandRef UATCommand = GetProfile()->FindOrAddUATCommand(*GetUATCommandInternalName());
-		OnUATCommandAdded(UATCommand);
-		UATCommand->SetOrder(GetUATCommandOrder());
-		BroadcastPropertyChanged();
-	}
-	void FUATCommandLaunchExtensionInstance::OnRemoved()
-	{
-		ILauncherProfileUATCommandPtr ExistingUATCommand = GetUATCommand();
-		if (ExistingUATCommand.IsValid())
-		{
-			OnUATCommandRemoved(ExistingUATCommand.ToSharedRef());
-			GetProfile()->RemoveUATCommand(*GetUATCommandInternalName());
-		}
-		BroadcastPropertyChanged();
-	}
-
-	bool FUATCommandLaunchExtensionInstance::IsUATCommandActive() const
-	{
-		return GetUATCommand() != nullptr;
-	}
-
-	ILauncherProfileUATCommandPtr FUATCommandLaunchExtensionInstance::GetUATCommand() const
-	{
-		return GetProfile()->GetUATCommand(*GetUATCommandInternalName());
-	}
-
-
-	bool FUATCommandLaunchExtensionInstance::CustomizeUATCommandLine( const ILauncherProfileUATCommandRef& InUATCommand, FString& InOutCommandLine )
-	{
-		if (InUATCommand == GetUATCommand())
-		{
-			CustomizeUATCommandLine(InOutCommandLine);
-			return true;
-		}
-
-		return false;
-	}
-
-
-
-
-
-
-
-
-
-	FAutomatedTestLaunchExtensionInstance::FAutomatedTestLaunchExtensionInstance( FArgs& InArgs )
-		: FLaunchExtensionInstance(InArgs)
-	{
-	}
-
-	void FAutomatedTestLaunchExtensionInstance::InternalInitialize()
-	{
-		ILauncherProfileAutomatedTestPtr AutomatedTest = GetTest();
-		if (AutomatedTest.IsValid())
-		{
-			OnTestAdded(AutomatedTest.ToSharedRef());
-
-			int32 Order = GetTestOrder();
-			check(Order > 0);
-			AutomatedTest->SetOrder(Order);
-		}
-	}
-
-
-	void FAutomatedTestLaunchExtensionInstance::OnAdded()
-	{
-		ILauncherProfileAutomatedTestRef AutomatedTest = GetProfile()->FindOrAddAutomatedTest(*GetTestInternalName());
-		OnTestAdded(AutomatedTest);
-
-		int32 Order = GetTestOrder();
-		check(Order > 0);
-		AutomatedTest->SetOrder(Order);
-
-		BroadcastPropertyChanged();
-
-	}
-
-	void FAutomatedTestLaunchExtensionInstance::OnRemoved()
-	{
-		ILauncherProfileAutomatedTestPtr ExistingAutomatedTest = GetTest();
-		if (ExistingAutomatedTest.IsValid())
-		{
-			OnTestRemoved(ExistingAutomatedTest.ToSharedRef());
-			GetProfile()->RemoveAutomatedTest(*GetTestInternalName());
-		}
-		BroadcastPropertyChanged();
-	}
-
-	bool FAutomatedTestLaunchExtensionInstance::IsTestActive() const
-	{
-		return GetTest() != nullptr;
-	}
-
-	ILauncherProfileAutomatedTestPtr FAutomatedTestLaunchExtensionInstance::GetTest() const
-	{
-		return GetProfile()->GetAutomatedTest(*GetTestInternalName());
-	}
-
-
-	bool FAutomatedTestLaunchExtensionInstance::CustomizeAutomatedTestCommandLine( const ILauncherProfileAutomatedTestRef& InAutomatedTest, FString& InOutCommandLine )
-	{
-		if (InAutomatedTest == GetTest())
-		{
-			CustomizeAutomatedTestCommandLine(InOutCommandLine);
-			return true;
-		}
-
-		return false;
-	}
 }
 
 #undef LOCTEXT_NAMESPACE
diff --git a/Engine/Plugins/Developer/ProjectLauncher/Source/ProjectLauncher/Private/Extension/UATCommandLaunchExtension.cpp b/Engine/Plugins/Developer/ProjectLauncher/Source/ProjectLauncher/Private/Extension/UATCommandLaunchExtension.cpp
new file mode 100644
--- /dev/null
+++ b/Engine/Plugins/Developer/ProjectLauncher/Source/ProjectLauncher/Private/Extension/UATCommandLaunchExtension.cpp
@@ -0,0 +1,75 @@
+// Copyright Epic Games, Inc. All Rights Reserved.
+
+#include "Extension/UATCommandLaunchExtension.h"
+
+#define LOCTEXT_NAMESPACE "FUATCommandLaunchExtensionInstance"
+
+
+namespace ProjectLauncher
+{
+	FUATCommandLaunchExtensionInstance::FUATCommandLaunchExtensionInstance( FArgs& InArgs )
+		: FLaunchExtensionInstance(InArgs)
+	{
+	}
+
+	void FUATCommandLaunchExtensionInstance::InternalInitialize()
+	{
+		ILauncherProfileUATCommandPtr UATCommand = GetUATCommand();
+		if (UATCommand.IsValid())
+		{
+			OnUATCommandAdded(UATCommand.ToSharedRef());
+			UATCommand->SetOrder(GetUATCommandOrder());
+		}
+	}
+
+	void FUATCommandLaunchExtensionInstance::OnAdded()
+	{
+		ILauncherProfileUATCommandRef UATCommand = GetProfile()->FindOrAddUATCommand(*GetUATCommandInternalName());
+		OnUATCommandAdded(UATCommand);
+		UATCommand->SetOrder(GetUATCommandOrder());
+		BroadcastPropertyChanged();
+	}
+	void FUATCommandLaunchExtensionInstance::OnRemoved()
+	{
+		ILauncherProfileUATCommandPtr ExistingUATCommand = GetUATCommand();
+		if (ExistingUATCommand.IsValid())
+		{
+			OnUATCommandRemoved(ExistingUATCommand.ToSharedRef());
+			GetProfile()->RemoveUATCommand(*GetUATCommandInternalName());
+		}
+		BroadcastPropertyChanged();
+	}
+
+	bool FUATCommandLaunchExtensionInstance::IsUATCommandActive() const
+	{
+		return GetUATCommand() != nullptr;
+	}
+
+	ILauncherProfileUATCommandPtr FUATCommandLaunchExtensionInstance::GetUATCommand() const
+	{
+		return GetProfile()->GetUATCommand(*GetUATCommandInternalName());
+	}
+
+
+	bool FUATCommandLaunchExtensionInstance::CustomizeUATCommandLine( const ILauncherProfileUATCommandRef& InUATCommand, FString& InOutCommandLine )
+	{
+		if (InUATCommand == GetUATCommand())
+		{
+			CustomizeUATCommandLine(InOutCommandLine);
+			return true;
+		}
+
+		return false;
+	}
+
+
+
+
+
+	void FUATCommandLaunchExtension::GetExtensionsMenuEntry( FExtensionsMenuEntry& MenuEntry ) const
+	{
+		MenuEntry = FExtensionsMenuEntry::UATCommands;
+	}
+}
+
+#undef LOCTEXT_NAMESPACE
diff --git a/Engine/Plugins/Developer/ProjectLauncher/Source/ProjectLauncher/Public/Extension/AutomatedTestLaunchExtension.h b/Engine/Plugins/Developer/ProjectLauncher/Source/ProjectLauncher/Public/Extension/AutomatedTestLaunchExtension.h
new file mode 100644
--- /dev/null
+++ b/Engine/Plugins/Developer/ProjectLauncher/Source/ProjectLauncher/Public/Extension/AutomatedTestLaunchExtension.h
@@ -0,0 +1,88 @@
+// Copyright Epic Games, Inc. All Rights Reserved.
+
+#pragma once
+
+#include "Extension/LaunchExtension.h"
+
+#define UE_API PROJECTLAUNCHER_API
+
+namespace ProjectLauncher
+{
+
+
+
+
+	/**
+	 * Helper base class for an automated test launch extension
+	 * Creates an item in the extension menu that will toggle the test on and off
+	 * note: It is not necessary for an automated test extensions to use this class
+	 */
+	class FAutomatedTestLaunchExtensionInstance : public FLaunchExtensionInstance
+	{
+	public:
+		UE_API FAutomatedTestLaunchExtensionInstance( FArgs& InArgs );
+
+		/**
+		 * Determine if this automated test is currently active
+		 */
+		UE_API bool IsTestActive() const;
+
+	protected:
+		UE_API virtual void OnAdded() override;
+		UE_API virtual void OnRemoved() override;
+
+		UE_API bool CustomizeAutomatedTestCommandLine( const ILauncherProfileAutomatedTestRef& InAutomatedTest, FString& InOutCommandLine ) override;
+		virtual bool HasCustomExtensionMenu() const { return true; }
+
+		/**
+		 * Advanced hook to allow any advanced modification of the command line when this extension's automated test is launched
+		 */
+		virtual void CustomizeAutomatedTestCommandLine( FString& InOutCommandLine ) {}
+
+		/**
+		 * Hook to allow any modification of the automated test & profile once it has been created. At least the test name should be specified
+		 */
+		virtual void OnTestAdded( ILauncherProfileAutomatedTestRef AutomatedTest ) = 0;
+
+		/**
+		 * Advanced hook to allow any modification of the profile just before the automated test is removed
+		 */
+		virtual void OnTestRemoved( ILauncherProfileAutomatedTestRef AutomatedTest ) {}
+
+		/**
+		 * Access this extension's automated test if it has been created
+		 */
+		UE_API ILauncherProfileAutomatedTestPtr GetTest() const;
+
+		/**
+		 * Internal name for the automated test. Must be globally unique, so recommended to include the extension name
+		 */
+		virtual const FString GetTestInternalName() const = 0;
+
+		/**
+		 * Determine the ordering for this test. lower number = run earlier
+		 */
+		virtual int32 GetTestOrder() const { return INT32_MAX; }
+
+		/**
+		 * Internal intitialization function, called immediately after construction - don't call directly
+		 */
+		UE_API virtual void InternalInitialize() override;
+
+	private:
+		virtual int32 GetHeadingSortOrder() const override { return GetTestOrder(); }
+	};
+
+
+
+
+
+	class FAutomatedTestLaunchExtension : public FLaunchExtension
+	{
+	public:
+		UE_API virtual void GetExtensionsMenuEntry( FExtensionsMenuEntry& MenuEntry ) const override;
+
+	};
+};
+
+#undef UE_API
diff --git a/Engine/Plugins/Developer/ProjectLauncher/Source/ProjectLauncher/Public/Extension/LaunchExtension.h b/Engine/Plugins/Developer/ProjectLauncher/Source/ProjectLauncher/Public/Extension/LaunchExtension.h
--- a/Engine/Plugins/Developer/ProjectLauncher/Source/ProjectLauncher/Public/Extension/LaunchExtension.h
+++ b/Engine/Plugins/Developer/ProjectLauncher/Source/ProjectLauncher/Public/Extension/LaunchExtension.h
@@ -332,132 +332,6 @@ namespace ProjectLauncher
 	};
 
 
-	/**
-	 * Helper base class for an UAT command launch extension
-	 * Creates an item in the extension menu that will toggle the UAT command on and off
-	 * note: It is not necessary for a UAT command launch extension to use this class
-	 */
-	class FUATCommandLaunchExtensionInstance : public FLaunchExtensionInstance
-	{
-	public:
-		UE_API FUATCommandLaunchExtensionInstance( FArgs& InArgs );
-
-		/**
-		 * Determine if this UAT Command is currently active
-		 */
-		UE_API bool IsUATCommandActive() const;
-
-	protected:
-		UE_API virtual void OnAdded() override;
-		UE_API virtual void OnRemoved() override;
-
-
-		UE_API bool CustomizeUATCommandLine( const ILauncherProfileUATCommandRef& InUATCommand, FString& InOutCommandLine ) override;
-		
-		/**
-		 * Advanced hook to allow any advanced modification of the command line when this extension's UAT Command is launched
-		 */
-		virtual void CustomizeUATCommandLine( FString& InOutCommandLine ) {}
-
-		/**
-		 * Hook to allow any modification of the UAT Command & profile once it has been created. At least the UAT Command name should be specified
-		 */
-		virtual void OnUATCommandAdded( ILauncherProfileUATCommandRef UATCommand ) = 0;
-
-		/**
-		 * Advanced hook to allow any modification of the profile just before the UAT Command is removed
-		 */
-		virtual void OnUATCommandRemoved( ILauncherProfileUATCommandRef UATCommand ) {}
-
-		/**
-		 * Access this extension's UAT Command if it has been created
-		 */
-		UE_API ILauncherProfileUATCommandPtr GetUATCommand() const;
-
-		/**
-		 * Internal name for the UAT command. Must be globally unique, so recommended to include the extension name
-		 */
-		virtual const FString GetUATCommandInternalName() const = 0;
-
-		/**
-		 * Determine the order of this UAT command, negative runs before BuildCookRun, positive runs after
-		 */
-		UE_API virtual int32 GetUATCommandOrder() const { return INT32_MAX; }
-
-		/**
-		 * Internal intitialization function, called immediately after construction - don't call directly
-		 */
-		UE_API virtual void InternalInitialize() override;
-
-	private:
-		virtual int32 GetHeadingSortOrder() const override { return GetUATCommandOrder(); }
-	};
-
-
-
-
-	/**
-	 * Helper base class for an automated test launch extension
-	 * Creates an item in the extension menu that will toggle the test on and off
-	 * note: It is not necessary for an automated test extensions to use this class
-	 */
-	class FAutomatedTestLaunchExtensionInstance : public FLaunchExtensionInstance
-	{
-	public:
-		UE_API FAutomatedTestLaunchExtensionInstance( FArgs& InArgs );
-
-		/**
-		 * Determine if this automated test is currently active
-		 */
-		UE_API bool IsTestActive() const;
-
-	protected:
-		UE_API virtual void OnAdded() override;
-		UE_API virtual void OnRemoved() override;
-
-		UE_API bool CustomizeAutomatedTestCommandLine( const ILauncherProfileAutomatedTestRef& InAutomatedTest, FString& InOutCommandLine ) override;
-		virtual bool HasCustomExtensionMenu() const { return true; }
-
-		/**
-		 * Advanced hook to allow any advanced modification of the command line when this extension's automated test is launched
-		 */
-		virtual void CustomizeAutomatedTestCommandLine( FString& InOutCommandLine ) {}
-
-		/**
-		 * Hook to allow any modification of the automated test & profile once it has been created. At least the test name should be specified
-		 */
-		virtual void OnTestAdded( ILauncherProfileAutomatedTestRef AutomatedTest ) = 0;
-
-		/**
-		 * Advanced hook to allow any modification of the profile just before the automated test is removed
-		 */
-		virtual void OnTestRemoved( ILauncherProfileAutomatedTestRef AutomatedTest ) {}
-
-		/**
-		 * Access this extension's automated test if it has been created
-		 */
-		UE_API ILauncherProfileAutomatedTestPtr GetTest() const;
-
-		/**
-		 * Internal name for the automated test. Must be globally unique, so recommended to include the extension name
-		 */
-		virtual const FString GetTestInternalName() const = 0;
-
-		/**
-		 * Determine the ordering for this test. lower number = run earlier
-		 */
-		virtual int32 GetTestOrder() const { return INT32_MAX; }
-
-		/**
-		 * Internal intitialization function, called immediately after construction - don't call directly
-		 */
-		UE_API virtual void InternalInitialize() override;
-
-	private:
-		virtual int32 GetHeadingSortOrder() const override { return GetTestOrder(); }
-	};
-
-
 
 
 
@@ -561,22 +435,6 @@ namespace ProjectLauncher
 		virtual TSharedPtr<FLaunchExtensionInstance> CreateInstanceForProfile( FLaunchExtensionInstance::FArgs& InArgs ) = 0;
 	};
 
-
-
-	class FAutomatedTestLaunchExtension : public FLaunchExtension
-	{
-	public:
-		UE_API virtual void GetExtensionsMenuEntry( FExtensionsMenuEntry& MenuEntry ) const override;
-
-	};
-
-	class FUATCommandLaunchExtension : public FLaunchExtension
-	{
-	public:
-		UE_API virtual void GetExtensionsMenuEntry( FExtensionsMenuEntry& MenuEntry ) const override;
-
-	};
-
 };
 
 #undef UE_API
diff --git a/Engine/Plugins/Developer/ProjectLauncher/Source/ProjectLauncher/Public/Extension/UATCommandLaunchExtension.h b/Engine/Plugins/Developer/ProjectLauncher/Source/ProjectLauncher/Public/Extension/UATCommandLaunchExtension.h
new file mode 100644
--- /dev/null
+++ b/Engine/Plugins/Developer/ProjectLauncher/Source/ProjectLauncher/Public/Extension/UATCommandLaunchExtension.h
@@ -0,0 +1,83 @@
+// Copyright Epic Games, Inc. All Rights Reserved.
+
+#pragma once
+
+#include "Extension/LaunchExtension.h"
+
+#define UE_API PROJECTLAUNCHER_API
+
+namespace ProjectLauncher
+{
+	/**
+	 * Helper base class for an UAT command launch extension
+	 * Creates an item in the extension menu that will toggle the UAT command on and off
+	 * note: It is not necessary for a UAT command launch extension to use this class
+	 */
+	class FUATCommandLaunchExtensionInstance : public FLaunchExtensionInstance
+	{
+	public:
+		UE_API FUATCommandLaunchExtensionInstance( FArgs& InArgs );
+
+		/**
+		 * Determine if this UAT Command is currently active
+		 */
+		UE_API bool IsUATCommandActive() const;
+
+	protected:
+		UE_API virtual void OnAdded() override;
+		UE_API virtual void OnRemoved() override;
+
+
+		UE_API bool CustomizeUATCommandLine( const ILauncherProfileUATCommandRef& InUATCommand, FString& InOutCommandLine ) override;
+		
+		/**
+		 * Advanced hook to allow any advanced modification of the command line when this extension's UAT Command is launched
+		 */
+		virtual void CustomizeUATCommandLine( FString& InOutCommandLine ) {}
+
+		/**
+		 * Hook to allow any modification of the UAT Command & profile once it has been created. At least the UAT Command name should be specified
+		 */
+		virtual void OnUATCommandAdded( ILauncherProfileUATCommandRef UATCommand ) = 0;
+
+		/**
+		 * Advanced hook to allow any modification of the profile just before the UAT Command is removed
+		 */
+		virtual void OnUATCommandRemoved( ILauncherProfileUATCommandRef UATCommand ) {}
+
+		/**
+		 * Access this extension's UAT Command if it has been created
+		 */
+		UE_API ILauncherProfileUATCommandPtr GetUATCommand() const;
+
+		/**
+		 * Internal name for the UAT command. Must be globally unique, so recommended to include the extension name
+		 */
+		virtual const FString GetUATCommandInternalName() const = 0;
+
+		/**
+		 * Determine the order of this UAT command, negative runs before BuildCookRun, positive runs after
+		 */
+		UE_API virtual int32 GetUATCommandOrder() const { return INT32_MAX; }
+
+		/**
+		 * Internal intitialization function, called immediately after construction - don't call directly
+		 */
+		UE_API virtual void InternalInitialize() override;
+
+	private:
+		virtual int32 GetHeadingSortOrder() const override { return GetUATCommandOrder(); }
+	};
+
+
+
+	class FUATCommandLaunchExtension : public FLaunchExtension
+	{
+	public:
+		UE_API virtual void GetExtensionsMenuEntry( FExtensionsMenuEntry& MenuEntry ) const override;
+
+	};
+
+};
+
+#undef UE_API
diff --git a/Engine/Plugins/Performance/AutomatedPerfTesting/Source/AutomatedPerfTesting/Private/LaunchExtension/AutomatedPerfTestLaunchExtension.h b/Engine/Plugins/Performance/AutomatedPerfTesting/Source/AutomatedPerfTesting/Private/LaunchExtension/AutomatedPerfTestLaunchExtension.h
--- a/Engine/Plugins/Performance/AutomatedPerfTesting/Source/AutomatedPerfTesting/Private/LaunchExtension/AutomatedPerfTestLaunchExtension.h
+++ b/Engine/Plugins/Performance/AutomatedPerfTesting/Source/AutomatedPerfTesting/Private/LaunchExtension/AutomatedPerfTestLaunchExtension.h
@@ -4,7 +4,7 @@
 
 #if WITH_EDITOR
 
-#include "Extension/LaunchExtension.h"
+#include "Extension/AutomatedTestLaunchExtension.h"
 #include "Templates/UnrealTemplate.h"
 
 enum class EAutomatedPerfTestType : uint8
