Project Launcher 2 support for multiple build cook run commands

#rnx
#Synced-CL 49676362

diff --git a/Engine/Plugins/Developer/ProjectLauncher/Source/CommonLaunchExtensions/Private/AdvancedCook/AdvancedCookLaunchExtension.cpp b/Engine/Plugins/Developer/ProjectLauncher/Source/CommonLaunchExtensions/Private/AdvancedCook/AdvancedCookLaunchExtension.cpp
--- a/Engine/Plugins/Developer/ProjectLauncher/Source/CommonLaunchExtensions/Private/AdvancedCook/AdvancedCookLaunchExtension.cpp
+++ b/Engine/Plugins/Developer/ProjectLauncher/Source/CommonLaunchExtensions/Private/AdvancedCook/AdvancedCookLaunchExtension.cpp
@@ -15,8 +15,37 @@
 
 
 FAdvancedCookLaunchExtensionInstance::FAdvancedCookLaunchExtensionInstance( FArgs& InArgs ) 
-	: ProjectLauncher::FLaunchExtensionInstance(InArgs)
+	: ProjectLauncher::FBuildCookRunLaunchExtensionInstance(InArgs)
 	, Owner(StaticCastSharedRef<FAdvancedCookLaunchExtension>(InArgs.Extension))
+{
+}
+
+TSharedRef<ProjectLauncher::FBuildCookRunExtension> FAdvancedCookLaunchExtensionInstance::CreateBuildCookRunExtension( const ProjectLauncher::FBuildCookRunExtension::FArgs& InArgs )
+{
+	return MakeShared<FBuildCookRunInstance>(InArgs);
+}
+
+
+bool FAdvancedCookLaunchExtensionInstance::IsBuildCookRunExtensionEnabledByDefault( const ILauncherProfileBuildCookRunRef& InBuildCookRun ) const
+{
+	return Owner->HasAdvancedCookOptions(InBuildCookRun);
+}
+
+bool FAdvancedCookLaunchExtensionInstance::CanToggleBuildCookRunExtension( const ILauncherProfileBuildCookRunRef& InBuildCookRun, bool bWantToEnable ) const
+{
+	if (!bWantToEnable && Owner->HasAdvancedCookOptions(InBuildCookRun))
+	{
+		return false;
+	}
+
+	return true;
+}
+
+
+
+
+FAdvancedCookLaunchExtensionInstance::FBuildCookRunInstance::FBuildCookRunInstance( const FArgs& InArgs )
+	: ProjectLauncher::FBuildCookRunExtension( InArgs )
 {
 	MapOption = GetProfile()->GetCookedMaps().Num() > 0 ? EMapOption::Selected : EMapOption::Startup;
 
@@ -25,20 +54,19 @@ FAdvancedCookLaunchExtensionInstance::FAdvancedCookLaunchExtensionInstance( FArg
 }
 
 
-void FAdvancedCookLaunchExtensionInstance::CustomizeTree( ProjectLauncher::FLaunchProfileTreeData& ProfileTreeData )
+void FAdvancedCookLaunchExtensionInstance::FBuildCookRunInstance::CustomizeTree( ProjectLauncher::FLaunchProfileTreeNode& ProfileTreeNode )
 {
 	auto IsVisible = [this]()
 	{
 		return (GetProfile()->GetCookMode() == ELauncherProfileCookModes::ByTheBook);
 	};
 
-	//ProfileTreeData.AddHeading( TEXT("Cooking"), LOCTEXT("CookingHeading", "Cook") )
-	AddDefaultHeading(ProfileTreeData)
+	AddDefaultHeading(ProfileTreeNode)
 		.AddBoolean( LOCTEXT("UnversionedCookLabel", "Unversioned Cook"),
 			{
-				.GetValue = [this]()			{ return GetProfile()->IsCookingUnversioned(); },
-				.SetValue = [this](bool bValue)	{ GetProfile()->SetUnversionedCooking(bValue); },
-				.GetDefaultValue = [this]()		{ return GetDefaultProfile()->IsCookingUnversioned(); },
+				.GetValue = [this]()			{ return GetBuildCookRun()->IsCookingUnversioned(); },
+				.SetValue = [this](bool bValue)	{ GetBuildCookRun()->SetUnversionedCooking(bValue); },
+				.GetDefaultValue = [this]()		{ return GetDefaultBuildCookRun()->IsCookingUnversioned(); },
 				.IsVisible = IsVisible,
 				.Validation = ProjectLauncher::FValidation({ELauncherProfileValidationErrors::UnversionedAndIncremental}),
 			}
@@ -57,8 +85,8 @@ void FAdvancedCookLaunchExtensionInstance::CustomizeTree( ProjectLauncher::FLaun
 		)
 		.AddString( LOCTEXT("AdditionalCookerOptionsLabel", "Additional Cooker Options"), 
 			{
-				.GetValue = [this]()				{ return GetProfile()->GetCookOptions(); },
-				.SetValue = [this](FString Value)	{ GetProfile()->SetCookOptions( Value ); },
+				.GetValue = [this]()				{ return GetBuildCookRun()->GetCookOptions(); },
+				.SetValue = [this](FString Value)	{ GetBuildCookRun()->SetCookOptions( Value ); },
 				.GetDefaultValue = []()				{ return FString(); },
 				.IsVisible = IsVisible,
 			}
@@ -70,12 +98,12 @@ void FAdvancedCookLaunchExtensionInstance::CustomizeTree( ProjectLauncher::FLaun
 
 
 BEGIN_SLATE_FUNCTION_BUILD_OPTIMIZATION
-TSharedRef<SWidget> FAdvancedCookLaunchExtensionInstance::CreateMapListWidget()
+TSharedRef<SWidget> FAdvancedCookLaunchExtensionInstance::FBuildCookRunInstance::CreateMapListWidget()
 {
 	MapListView = 
 		SNew(SCustomLaunchMapListView, GetModel())
-		.OnSelectionChanged(this, &FAdvancedCookLaunchExtensionInstance::SetMapsToCook)
-		.SelectedMaps(this, &FAdvancedCookLaunchExtensionInstance::GetMapsToCook)
+		.OnSelectionChanged(this, &FAdvancedCookLaunchExtensionInstance::FBuildCookRunInstance::SetMapsToCook)
+		.SelectedMaps(this, &FAdvancedCookLaunchExtensionInstance::FBuildCookRunInstance::GetMapsToCook)
 		.ProjectPath_Lambda( [this]() { return GetModel()->GetProjectPath(GetProfile()); })
 	;
 
@@ -95,8 +123,8 @@ TSharedRef<SWidget> FAdvancedCookLaunchExtensionInstance::CreateMapListWidget()
 			.VAlign(VAlign_Center)
 			[
 				SNew(SSegmentedControl<EMapOption>)
-				.Value( this, &FAdvancedCookLaunchExtensionInstance::GetMapOption)
-				.OnValueChanged(this, &FAdvancedCookLaunchExtensionInstance::SetMapOption)
+				.Value( this, &FAdvancedCookLaunchExtensionInstance::FBuildCookRunInstance::GetMapOption)
+				.OnValueChanged(this, &FAdvancedCookLaunchExtensionInstance::FBuildCookRunInstance::SetMapOption)
 
 				+SSegmentedControl<EMapOption>::Slot(EMapOption::Startup)
 				.HAlign(HAlign_Center)
@@ -138,8 +166,8 @@ TSharedRef<SWidget> FAdvancedCookLaunchExtensionInstance::CreateMapListWidget()
 			SNew(SVerticalResizeBox)
 			.Visibility_Lambda([this](){ return (GetMapOption() == EMapOption::Selected) ? EVisibility::Visible : EVisibility::Collapsed; } )
 			.HandleHeight(4)
-			.ContentHeight(this, &FAdvancedCookLaunchExtensionInstance::GetMapListHeight)
-			.ContentHeightChanged(this, &FAdvancedCookLaunchExtensionInstance::SetMapListHeight)
+			.ContentHeight(this, &FAdvancedCookLaunchExtensionInstance::FBuildCookRunInstance::GetMapListHeight)
+			.ContentHeightChanged(this, &FAdvancedCookLaunchExtensionInstance::FBuildCookRunInstance::SetMapListHeight)
 			.HandleColor( FAppStyle::Get().GetSlateColor("Colors.Secondary").GetSpecifiedColor() )
 			[
 				MapListView.ToSharedRef()
@@ -152,12 +180,12 @@ END_SLATE_FUNCTION_BUILD_OPTIMIZATION
 
 
 BEGIN_SLATE_FUNCTION_BUILD_OPTIMIZATION
-TSharedRef<SWidget> FAdvancedCookLaunchExtensionInstance::CreateCultureListWidget()
+TSharedRef<SWidget> FAdvancedCookLaunchExtensionInstance::FBuildCookRunInstance::CreateCultureListWidget()
 {
 	CultureListView = 
 		SNew(SCustomLaunchCultureListView, GetModel())
-		.OnSelectionChanged(this, &FAdvancedCookLaunchExtensionInstance::SetCulturesToCook)
-		.SelectedCultures(this, &FAdvancedCookLaunchExtensionInstance::GetCulturesToCook)
+		.OnSelectionChanged(this, &FAdvancedCookLaunchExtensionInstance::FBuildCookRunInstance::SetCulturesToCook)
+		.SelectedCultures(this, &FAdvancedCookLaunchExtensionInstance::FBuildCookRunInstance::GetCulturesToCook)
 		.ProjectPath_Lambda( [this]() { return GetModel()->GetProjectPath(GetProfile()); })
 	;
 
@@ -177,8 +205,8 @@ TSharedRef<SWidget> FAdvancedCookLaunchExtensionInstance::CreateCultureListWidge
 			.VAlign(VAlign_Center)
 			[
 				SNew(SSegmentedControl<ECultureOption>)
-				.Value( this, &FAdvancedCookLaunchExtensionInstance::GetCultureOption)
-				.OnValueChanged(this, &FAdvancedCookLaunchExtensionInstance::SetCultureOption)
+				.Value( this, &FAdvancedCookLaunchExtensionInstance::FBuildCookRunInstance::GetCultureOption)
+				.OnValueChanged(this, &FAdvancedCookLaunchExtensionInstance::FBuildCookRunInstance::SetCultureOption)
 
 				+SSegmentedControl<ECultureOption>::Slot(ECultureOption::Default)
 				.HAlign(HAlign_Center)
@@ -220,8 +248,8 @@ TSharedRef<SWidget> FAdvancedCookLaunchExtensionInstance::CreateCultureListWidge
 			SNew(SVerticalResizeBox)
 			.Visibility_Lambda([this](){ return GetCultureOption() == ECultureOption::Custom ? EVisibility::Visible : EVisibility::Collapsed; } )
 			.HandleHeight(4)
-			.ContentHeight(this, &FAdvancedCookLaunchExtensionInstance::GetCultureListHeight)
-			.ContentHeightChanged(this, &FAdvancedCookLaunchExtensionInstance::SetCultureListHeight)
+			.ContentHeight(this, &FAdvancedCookLaunchExtensionInstance::FBuildCookRunInstance::GetCultureListHeight)
+			.ContentHeightChanged(this, &FAdvancedCookLaunchExtensionInstance::FBuildCookRunInstance::SetCultureListHeight)
 			.HandleColor( FAppStyle::Get().GetSlateColor("Colors.Secondary").GetSpecifiedColor() )
 			[
 				CultureListView.ToSharedRef()
@@ -232,42 +260,42 @@ TSharedRef<SWidget> FAdvancedCookLaunchExtensionInstance::CreateCultureListWidge
 END_SLATE_FUNCTION_BUILD_OPTIMIZATION
 
 
-void FAdvancedCookLaunchExtensionInstance::SetMapsToCook(TArray<FString> MapsToCook)
+void FAdvancedCookLaunchExtensionInstance::FBuildCookRunInstance::SetMapsToCook(TArray<FString> MapsToCook)
 {
-	GetProfile()->ClearCookedMaps(); 
+	GetBuildCookRun()->ClearCookedMaps(); 
 	for (const FString& Map : MapsToCook)
 	{
-		GetProfile()->AddCookedMap(Map);
+		GetBuildCookRun()->AddCookedMap(Map);
 	}
 
 	OnPropertyChanged();
 }
 
-TArray<FString> FAdvancedCookLaunchExtensionInstance::GetMapsToCook() const
+TArray<FString> FAdvancedCookLaunchExtensionInstance::FBuildCookRunInstance::GetMapsToCook() const
 {
-	return GetProfile()->GetCookedMaps();
+	return GetBuildCookRun()->GetCookedMaps();
 }
 
-float FAdvancedCookLaunchExtensionInstance::GetMapListHeight() const
+float FAdvancedCookLaunchExtensionInstance::FBuildCookRunInstance::GetMapListHeight() const
 {
 	return MapListHeight;
 }
 
-void FAdvancedCookLaunchExtensionInstance::SetMapListHeight( float NewHeight )
+void FAdvancedCookLaunchExtensionInstance::FBuildCookRunInstance::SetMapListHeight( float NewHeight )
 {
 	static const float MinMapListHeight = 100.0f;
 
 	MapListHeight = FMath::Max(NewHeight, MinMapListHeight);
 
-	RequestTreeRefresh();
+	GetExtensionInstance()->RequestTreeRefresh();
 }
 
-FAdvancedCookLaunchExtensionInstance::EMapOption FAdvancedCookLaunchExtensionInstance::GetMapOption() const
+FAdvancedCookLaunchExtensionInstance::FBuildCookRunInstance::EMapOption FAdvancedCookLaunchExtensionInstance::FBuildCookRunInstance::GetMapOption() const
 {
 	return MapOption;
 }
 
-void FAdvancedCookLaunchExtensionInstance::SetMapOption( EMapOption NewMapOption )
+void FAdvancedCookLaunchExtensionInstance::FBuildCookRunInstance::SetMapOption( EMapOption NewMapOption )
 {
 	bool bShow = (NewMapOption == EMapOption::Selected);
 
@@ -276,7 +304,7 @@ void FAdvancedCookLaunchExtensionInstance::SetMapOption( EMapOption NewMapOption
 	if (bShow)
 	{
 		// restore the cooked maps again, if any
-		if (CachedMapsToCook.Num() > 0 && GetProfile()->GetCookedMaps().Num() == 0)
+		if (CachedMapsToCook.Num() > 0 && GetBuildCookRun()->GetCookedMaps().Num() == 0)
 		{
 			SetMapsToCook(CachedMapsToCook);
 			CachedMapsToCook.Reset();
@@ -285,7 +313,7 @@ void FAdvancedCookLaunchExtensionInstance::SetMapOption( EMapOption NewMapOption
 	else
 	{
 		// to set the 'cook startup maps only', its necessary to remove all the cooked maps - take a copy of the values to allow it to be restored
-		CachedMapsToCook = GetProfile()->GetCookedMaps();
+		CachedMapsToCook = GetBuildCookRun()->GetCookedMaps();
 		SetMapsToCook(TArray<FString>());
 	}
 
@@ -297,42 +325,42 @@ void FAdvancedCookLaunchExtensionInstance::SetMapOption( EMapOption NewMapOption
 
 
 
-void FAdvancedCookLaunchExtensionInstance::SetCulturesToCook(TArray<FString> CulturesToCook)
+void FAdvancedCookLaunchExtensionInstance::FBuildCookRunInstance::SetCulturesToCook(TArray<FString> CulturesToCook)
 {
-	GetProfile()->ClearCookedCultures(); 
+	GetBuildCookRun()->ClearCookedCultures(); 
 	for (const FString& Culture : CulturesToCook)
 	{
-		GetProfile()->AddCookedCulture(Culture);
+		GetBuildCookRun()->AddCookedCulture(Culture);
 	}
 
 	OnPropertyChanged();
 }
 
-TArray<FString> FAdvancedCookLaunchExtensionInstance::GetCulturesToCook() const
+TArray<FString> FAdvancedCookLaunchExtensionInstance::FBuildCookRunInstance::GetCulturesToCook() const
 {
-	return GetProfile()->GetCookedCultures();
+	return GetBuildCookRun()->GetCookedCultures();
 }
 
-float FAdvancedCookLaunchExtensionInstance::GetCultureListHeight() const
+float FAdvancedCookLaunchExtensionInstance::FBuildCookRunInstance::GetCultureListHeight() const
 {
 	return CultureListHeight;
 }
 
-void FAdvancedCookLaunchExtensionInstance::SetCultureListHeight( float NewHeight )
+void FAdvancedCookLaunchExtensionInstance::FBuildCookRunInstance::SetCultureListHeight( float NewHeight )
 {
 	static const float MinCultureListHeight = 100.0f;
 
 	CultureListHeight = FMath::Max(NewHeight, MinCultureListHeight);
 
-	RequestTreeRefresh();
+	GetExtensionInstance()->RequestTreeRefresh();
 }
 
-FAdvancedCookLaunchExtensionInstance::ECultureOption FAdvancedCookLaunchExtensionInstance::GetCultureOption() const
+FAdvancedCookLaunchExtensionInstance::FBuildCookRunInstance::ECultureOption FAdvancedCookLaunchExtensionInstance::FBuildCookRunInstance::GetCultureOption() const
 {
 	return CultureOption;
 }
 
-void FAdvancedCookLaunchExtensionInstance::SetCultureOption( ECultureOption NewCultureOption )
+void FAdvancedCookLaunchExtensionInstance::FBuildCookRunInstance::SetCultureOption( ECultureOption NewCultureOption )
 {
 	bool bShow = (NewCultureOption == ECultureOption::Custom);
 
@@ -347,7 +375,7 @@ void FAdvancedCookLaunchExtensionInstance::SetCultureOption( ECultureOption NewC
 	else
 	{
 		// to set the 'english only', its necessary to remove all the currently selected cultures - take a copy of the values to allow it to be restored
-		CachedCulturesToCook = GetProfile()->GetCookedCultures();
+		CachedCulturesToCook = GetBuildCookRun()->GetCookedCultures();
 		SetCulturesToCook(TArray<FString>{TEXT("en")});
 	}
 
@@ -357,7 +385,7 @@ void FAdvancedCookLaunchExtensionInstance::SetCultureOption( ECultureOption NewC
 }
 
 
-void FAdvancedCookLaunchExtensionInstance::OnProjectChanged()
+void FAdvancedCookLaunchExtensionInstance::FBuildCookRunInstance::OnProjectChanged()
 {
 	if (MapListView.IsValid())
 	{
@@ -365,11 +393,6 @@ void FAdvancedCookLaunchExtensionInstance::OnProjectChanged()
 	}
 }
 
-bool FAdvancedCookLaunchExtensionInstance::CanBeRemoved() const
-{
-	return !Owner->HasAdvancedCookOptions(GetProfile());
-}
-
 
 
 TSharedPtr<ProjectLauncher::FLaunchExtensionInstance> FAdvancedCookLaunchExtension::CreateInstanceForProfile( ProjectLauncher::FLaunchExtensionInstance::FArgs& InArgs )
@@ -392,14 +415,14 @@ bool FAdvancedCookLaunchExtension::IsCreatedByDefault( ILauncherProfileRef InPro
 	return HasAdvancedCookOptions(InProfile);
 }
 
-bool FAdvancedCookLaunchExtension::HasAdvancedCookOptions(ILauncherProfileRef InProfile) const
+bool FAdvancedCookLaunchExtension::HasAdvancedCookOptions(const ILauncherProfileBuildCookRunRef& InBuildCookRun) const
 {
-	if (!InProfile->IsCookingUnversioned() || !InProfile->GetCookOptions().IsEmpty() || InProfile->GetCookedMaps().Num() > 0)
+	if (!InBuildCookRun->IsCookingUnversioned() || !InBuildCookRun->GetCookOptions().IsEmpty() || InBuildCookRun->GetCookedMaps().Num() > 0)
 	{
 		return true;
 	}
 
-	bool bHasDefaultCultures = (InProfile->GetCookedCultures().Num() == 1 && InProfile->GetCookedCultures()[0] == TEXT("en"));
+	bool bHasDefaultCultures = (InBuildCookRun->GetCookedCultures().Num() == 1 && InBuildCookRun->GetCookedCultures()[0] == TEXT("en"));
 	if (!bHasDefaultCultures)
 	{
 		return true;
@@ -409,6 +432,19 @@ bool FAdvancedCookLaunchExtension::HasAdvancedCookOptions(ILauncherProfileRef In
 	return false;
 }
 
+bool FAdvancedCookLaunchExtension::HasAdvancedCookOptions(ILauncherProfileRef InProfile) const
+{
+	for (const ILauncherProfileBuildCookRunRef& BuildCookRun : InProfile->GetBuildCookRunCommands())
+	{
+		if (HasAdvancedCookOptions(BuildCookRun))
+		{
+			return true;
+		}
+	}
+
+	return false;
+}
+
 
 
 
diff --git a/Engine/Plugins/Developer/ProjectLauncher/Source/CommonLaunchExtensions/Private/AdvancedCook/AdvancedCookLaunchExtension.h b/Engine/Plugins/Developer/ProjectLauncher/Source/CommonLaunchExtensions/Private/AdvancedCook/AdvancedCookLaunchExtension.h
--- a/Engine/Plugins/Developer/ProjectLauncher/Source/CommonLaunchExtensions/Private/AdvancedCook/AdvancedCookLaunchExtension.h
+++ b/Engine/Plugins/Developer/ProjectLauncher/Source/CommonLaunchExtensions/Private/AdvancedCook/AdvancedCookLaunchExtension.h
@@ -2,71 +2,81 @@
 
 #pragma once
 
-#include "Extension/LaunchExtension.h"
+#include "Extension/BuildCookRunLaunchExtension.h"
 
 
 class SCustomLaunchMapListView;
 class SCustomLaunchCultureListView;
 
 
-class FAdvancedCookLaunchExtensionInstance : public ProjectLauncher::FLaunchExtensionInstance
+class FAdvancedCookLaunchExtensionInstance : public ProjectLauncher::FBuildCookRunLaunchExtensionInstance
 {
 public:
 	FAdvancedCookLaunchExtensionInstance( FArgs& InArgs );
 	virtual ~FAdvancedCookLaunchExtensionInstance() = default;
 
-	virtual void CustomizeTree( ProjectLauncher::FLaunchProfileTreeData& ProfileTreeData ) override;
-	virtual bool CanBeRemoved() const override;
+	virtual TSharedRef<ProjectLauncher::FBuildCookRunExtension> CreateBuildCookRunExtension( const ProjectLauncher::FBuildCookRunExtension::FArgs& InArgs ) override;
 
-protected:
-	TSharedRef<class FAdvancedCookLaunchExtension> Owner;
-
-
-	TArray<FString> CachedMapsToCook;
-	enum class EMapOption : uint8
-	{
-		Startup,
-		Selected,
-	};
-	EMapOption MapOption = EMapOption::Startup;
-
-	void SetMapsToCook(TArray<FString> MapsToCook);
-	TArray<FString> GetMapsToCook() const;
-	EMapOption GetMapOption() const;
-	void SetMapOption( EMapOption MapOption );
-	float GetMapListHeight() const;
-	void SetMapListHeight( float NewHeight );
-	TSharedRef<SWidget> CreateMapListWidget();
+	virtual bool IsBuildCookRunExtensionEnabledByDefault( const ILauncherProfileBuildCookRunRef& InBuildCookRun ) const override;
+	virtual bool CanToggleBuildCookRunExtension( const ILauncherProfileBuildCookRunRef& InBuildCookRun, bool bWantToEnable ) const override;
 
-	float MapListHeight = 300.0f;
-	TSharedPtr<SCustomLaunchMapListView> MapListView;
 
+protected:
+	TSharedRef<class FAdvancedCookLaunchExtension> Owner;
 
-	TArray<FString> CachedCulturesToCook;
-	enum class ECultureOption : uint8
+	class FBuildCookRunInstance : public ProjectLauncher::FBuildCookRunExtension
 	{
-		Default,
-		Custom,
+	public:
+		FBuildCookRunInstance(const ProjectLauncher::FBuildCookRunExtension::FArgs& InArgs );
+		virtual void CustomizeTree( ProjectLauncher::FLaunchProfileTreeNode& ProfileTreeNode ) override;
+
+
+		TArray<FString> CachedMapsToCook;
+		enum class EMapOption : uint8
+		{
+			Startup,
+			Selected,
+		};
+		EMapOption MapOption = EMapOption::Startup;
+
+		void SetMapsToCook(TArray<FString> MapsToCook);
+		TArray<FString> GetMapsToCook() const;
+		EMapOption GetMapOption() const;
+		void SetMapOption( EMapOption MapOption );
+		float GetMapListHeight() const;
+		void SetMapListHeight( float NewHeight );
+		TSharedRef<SWidget> CreateMapListWidget();
+
+		float MapListHeight = 300.0f;
+		TSharedPtr<SCustomLaunchMapListView> MapListView;
+
+
+		TArray<FString> CachedCulturesToCook;
+		enum class ECultureOption : uint8
+		{
+			Default,
+			Custom,
+		};
+		ECultureOption CultureOption = ECultureOption::Default;
+
+		void SetCulturesToCook(TArray<FString> CulturesToCook);
+		TArray<FString> GetCulturesToCook() const;
+		ECultureOption GetCultureOption() const;
+		void SetCultureOption( ECultureOption CultureOption );
+		float GetCultureListHeight() const;
+		void SetCultureListHeight( float NewHeight );
+		TSharedRef<SWidget> CreateCultureListWidget();
+
+		float CultureListHeight = 150.0f;
+		TSharedPtr<SCustomLaunchCultureListView> CultureListView;
+
+
+		virtual void OnProjectChanged() override;
 	};
-	ECultureOption CultureOption = ECultureOption::Default;
-
-	void SetCulturesToCook(TArray<FString> CulturesToCook);
-	TArray<FString> GetCulturesToCook() const;
-	ECultureOption GetCultureOption() const;
-	void SetCultureOption( ECultureOption CultureOption );
-	float GetCultureListHeight() const;
-	void SetCultureListHeight( float NewHeight );
-	TSharedRef<SWidget> CreateCultureListWidget();
-
-	float CultureListHeight = 150.0f;
-	TSharedPtr<SCustomLaunchCultureListView> CultureListView;
-
-
-	virtual void OnProjectChanged() override;
 };
 
 
-class FAdvancedCookLaunchExtension : public ProjectLauncher::FLaunchExtension
+class FAdvancedCookLaunchExtension : public ProjectLauncher::FBuildCookRunLaunchExtension
 {
 public:
 	virtual TSharedPtr<ProjectLauncher::FLaunchExtensionInstance> CreateInstanceForProfile( ProjectLauncher::FLaunchExtensionInstance::FArgs& InArgs ) override;
@@ -74,5 +84,6 @@ public:
 	virtual FText GetDisplayName() const override;
 	virtual bool IsCreatedByDefault( ILauncherProfileRef InProfile, TSharedRef<ProjectLauncher::FModel> InModel ) const override;
 
+	bool HasAdvancedCookOptions(const ILauncherProfileBuildCookRunRef& InBuildCookRun) const;
 	bool HasAdvancedCookOptions(ILauncherProfileRef InProfile) const;
 };
\ No newline at end of file
diff --git a/Engine/Plugins/Developer/ProjectLauncher/Source/CommonLaunchExtensions/Private/BootTest/BootTestLaunchExtension.cpp b/Engine/Plugins/Developer/ProjectLauncher/Source/CommonLaunchExtensions/Private/BootTest/BootTestLaunchExtension.cpp
--- a/Engine/Plugins/Developer/ProjectLauncher/Source/CommonLaunchExtensions/Private/BootTest/BootTestLaunchExtension.cpp
+++ b/Engine/Plugins/Developer/ProjectLauncher/Source/CommonLaunchExtensions/Private/BootTest/BootTestLaunchExtension.cpp
@@ -12,17 +12,11 @@ void FBootTestLaunchExtensionInstance::CustomizeTree( ProjectLauncher::FLaunchPr
 {
 	// note: Windowed option is here mostly as an example of adding UI items for a particular profile - your tests may not need it
 
-	auto IsVisible = [this]() // this always outlives the ProfileTreeData
-	{
-		return IsTestActive();
-	};
-
 	AddDefaultHeading(ProfileTreeData)
 		.AddBoolean(LOCTEXT("WindowedLabel","Windowed"),
 		{
 			.GetValue = [this]()          { return GetConfigBool(EConfig::PerProfile, TEXT("Windowed")); },
 			.SetValue = [this](bool bVal) { SetConfigBool(EConfig::PerProfile, TEXT("Windowed"), bVal ); },
-			.IsVisible = IsVisible,
 		}
 	);
 }
diff --git a/Engine/Plugins/Developer/ProjectLauncher/Source/CommonLaunchExtensions/Private/DeprecatedProperties/DeprecatedPropertiesLaunchExtension.cpp b/Engine/Plugins/Developer/ProjectLauncher/Source/CommonLaunchExtensions/Private/DeprecatedProperties/DeprecatedPropertiesLaunchExtension.cpp
--- a/Engine/Plugins/Developer/ProjectLauncher/Source/CommonLaunchExtensions/Private/DeprecatedProperties/DeprecatedPropertiesLaunchExtension.cpp
+++ b/Engine/Plugins/Developer/ProjectLauncher/Source/CommonLaunchExtensions/Private/DeprecatedProperties/DeprecatedPropertiesLaunchExtension.cpp
@@ -14,21 +14,47 @@
 
 
 FDeprecatedPropertiesLaunchExtensionInstance::FDeprecatedPropertiesLaunchExtensionInstance( FArgs& InArgs ) 
-	: ProjectLauncher::FLaunchExtensionInstance(InArgs)
+	: ProjectLauncher::FBuildCookRunLaunchExtensionInstance(InArgs)
 	, Owner(StaticCastSharedRef<FDeprecatedPropertiesLaunchExtension>(InArgs.Extension))
 {
 }
 
 
-void FDeprecatedPropertiesLaunchExtensionInstance::CustomizeTree( ProjectLauncher::FLaunchProfileTreeData& ProfileTreeData )
+TSharedRef<ProjectLauncher::FBuildCookRunExtension> FDeprecatedPropertiesLaunchExtensionInstance::CreateBuildCookRunExtension( const ProjectLauncher::FBuildCookRunExtension::FArgs& InArgs )
 {
-	AddDefaultHeading(ProfileTreeData)
+	return MakeShared<FBuildCookRunInstance>(InArgs);
+}
+
+bool FDeprecatedPropertiesLaunchExtensionInstance::IsBuildCookRunExtensionEnabledByDefault( const ILauncherProfileBuildCookRunRef& InBuildCookRun ) const
+{
+	return Owner->HasDeprecatedProperties(InBuildCookRun, GetModel());
+}
+
+bool FDeprecatedPropertiesLaunchExtensionInstance::CanToggleBuildCookRunExtension( const ILauncherProfileBuildCookRunRef& InBuildCookRun, bool bWantToEnable ) const
+{
+	if (!bWantToEnable && Owner->HasDeprecatedProperties(InBuildCookRun, GetModel()))
+	{
+		return false;
+	}
+	return true;
+}
+
+
+FDeprecatedPropertiesLaunchExtensionInstance::FBuildCookRunInstance::FBuildCookRunInstance( const FArgs& InArgs )
+	: ProjectLauncher::FBuildCookRunExtension( InArgs )
+	, Owner(StaticCastSharedRef<FDeprecatedPropertiesLaunchExtension>(InArgs.Extension))
+{
+}
+
+void FDeprecatedPropertiesLaunchExtensionInstance::FBuildCookRunInstance::CustomizeTree( ProjectLauncher::FLaunchProfileTreeNode& ProfileTreeNode )
+{
+	AddDefaultHeading(ProfileTreeNode)
 		.AddWidget( FText::GetEmpty(),
 			SNew(SBorder)
 			.BorderImage(FAppStyle::GetBrush("ToolPanel.LightGroupBorder"))
 			.BorderBackgroundColor_Lambda( [this]()
 			{
-				return Owner->HasDeprecatedProperties(GetProfile(), GetDefaultProfile())
+				return Owner->HasDeprecatedProperties(GetBuildCookRun(), GetModel())
 					? FAppStyle::Get().GetSlateColor("Colors.Warning")
 					: FAppStyle::Get().GetSlateColor("Colors.AccentGreen");
 			})
@@ -37,7 +63,7 @@ void FDeprecatedPropertiesLaunchExtensionInstance::CustomizeTree( ProjectLaunche
 				.ColorAndOpacity(FLinearColor::White)
 				.Text_Lambda( [this]() 
 				{ 
-					return Owner->HasDeprecatedProperties(GetProfile(), GetDefaultProfile()) 
+					return Owner->HasDeprecatedProperties(GetBuildCookRun(), GetModel())
 						? LOCTEXT("DeprecatedItemsBanner", "These properties may be removed in future engine releases.") 
 						: LOCTEXT("NoDeprecatedItemsBanner", "These properties may be removed in future engine releases. You are not using any at the moment.");
 				})
@@ -47,87 +73,81 @@ void FDeprecatedPropertiesLaunchExtensionInstance::CustomizeTree( ProjectLaunche
 		// ... http chunks ...
 		.AddBoolean( LOCTEXT("CreateHttpChunksLabel", "Create Http Chunk Install Data"),
 			{
-				.GetValue = [this]()				{ return GetProfile()->IsGenerateHttpChunkData(); },
-				.SetValue = [this](bool bValue)		{ GetProfile()->SetGenerateHttpChunkData(bValue); },
-				.GetDefaultValue = [this]()			{ return GetDefaultProfile()->IsGenerateHttpChunkData(); },
+				.GetValue = [this]()				{ return GetBuildCookRun()->IsGenerateHttpChunkData(); },
+				.SetValue = [this](bool bValue)		{ GetBuildCookRun()->SetGenerateHttpChunkData(bValue); },
+				.GetDefaultValue = [this]()			{ return GetDefaultBuildCookRun()->IsGenerateHttpChunkData(); },
 			}
 		)
 		.AddDirectoryString( LOCTEXT("HttpChunkDirLabel", "Http Chunk Install Data Path"),
 			{
-				.GetValue = [this]()				{ return GetProfile()->GetHttpChunkDataDirectory(); },
-				.SetValue = [this](FString Value)	{ GetProfile()->SetHttpChunkDataDirectory(Value); },
-				.GetDefaultValue = [this]()			{ return GetDefaultProfile()->GetHttpChunkDataDirectory(); },
-				.IsEnabled = [this]()				{ return GetProfile()->IsGenerateHttpChunkData(); },
+				.GetValue = [this]()				{ return GetBuildCookRun()->GetHttpChunkDataDirectory(); },
+				.SetValue = [this](FString Value)	{ GetBuildCookRun()->SetHttpChunkDataDirectory(Value); },
+				.GetDefaultValue = [this]()			{ return GetDefaultBuildCookRun()->GetHttpChunkDataDirectory(); },
+				.IsEnabled = [this]()				{ return GetBuildCookRun()->IsGenerateHttpChunkData(); },
 			}
 		)
 		.AddString( LOCTEXT("HttpChunkReleaseName", "Http Chunk Install Release Name"),
 			{
-				.GetValue = [this]()				{ return GetProfile()->GetHttpChunkDataReleaseName(); },
-				.SetValue = [this](FString Value)	{ GetProfile()->SetHttpChunkDataReleaseName(Value); },
-				.GetDefaultValue = [this]()			{ return GetDefaultProfile()->GetHttpChunkDataReleaseName(); },
-				.IsEnabled = [this]()				{ return GetProfile()->IsGenerateHttpChunkData(); },
+				.GetValue = [this]()				{ return GetBuildCookRun()->GetHttpChunkDataReleaseName(); },
+				.SetValue = [this](FString Value)	{ GetBuildCookRun()->SetHttpChunkDataReleaseName(Value); },
+				.GetDefaultValue = [this]()			{ return GetDefaultBuildCookRun()->GetHttpChunkDataReleaseName(); },
+				.IsEnabled = [this]()				{ return GetBuildCookRun()->IsGenerateHttpChunkData(); },
 			}
 		)
 
 		// ... misc items ...
 		.AddBoolean( LOCTEXT("SkipEditorContentLabel", "Don't include editor content in the build"),
 			{
-				.GetValue = [this]()				{ return GetProfile()->GetSkipCookingEditorContent(); },
-				.SetValue = [this](bool bValue)		{ GetProfile()->SetSkipCookingEditorContent(bValue); },
-				.GetDefaultValue = [this]()			{ return GetDefaultProfile()->GetSkipCookingEditorContent(); },
+				.GetValue = [this]()				{ return GetBuildCookRun()->GetSkipCookingEditorContent(); },
+				.SetValue = [this](bool bValue)		{ GetBuildCookRun()->SetSkipCookingEditorContent(bValue); },
+				.GetDefaultValue = [this]()			{ return GetDefaultBuildCookRun()->GetSkipCookingEditorContent(); },
 			}
 		)
 		.AddWidget( LOCTEXT("CookerConfigLabel", "Cooker build configuration"), 
 			{
-				.IsDefault = [this]()				{ return GetProfile()->GetCookConfiguration() == GetDefaultProfile()->GetCookConfiguration(); },
-				.SetToDefault = [this]()			{ return GetProfile()->SetCookConfiguration( GetDefaultProfile()->GetCookConfiguration() ); },
+				.IsDefault = [this]()				{ return GetBuildCookRun()->GetCookConfiguration() == GetDefaultBuildCookRun()->GetCookConfiguration(); },
+				.SetToDefault = [this]()			{ return GetBuildCookRun()->SetCookConfiguration( GetDefaultBuildCookRun()->GetCookConfiguration() ); },
 			},
 			SNew(SCustomLaunchLexToStringCombo<EBuildConfiguration>)
-			.OnSelectionChanged_Lambda( [this](EBuildConfiguration Value) { GetProfile()->SetCookConfiguration(Value); } )
-			.SelectedItem_Lambda( [this]() { return GetProfile()->GetCookConfiguration(); } )
+			.OnSelectionChanged_Lambda( [this](EBuildConfiguration Value) { GetBuildCookRun()->SetCookConfiguration(Value); } )
+			.SelectedItem_Lambda( [this]() { return GetBuildCookRun()->GetCookConfiguration(); } )
 			.Items(TArray<EBuildConfiguration>{ EBuildConfiguration::Debug, EBuildConfiguration::DebugGame, EBuildConfiguration::Development, EBuildConfiguration::Test, EBuildConfiguration::Shipping })
 		)
 		.AddDirectoryString( LOCTEXT("PackageDirLabel", "Staging Directory Override"),
 			{
-				.GetValue = [this]()				{ return GetProfile()->GetPackageDirectory(); },
-				.SetValue = [this](FString Value)	{ GetProfile()->SetPackageDirectory(Value); },
-				.GetDefaultValue = [this]()			{ return GetDefaultProfile()->GetPackageDirectory(); },
+				.GetValue = [this]()				{ return GetBuildCookRun()->GetPackageDirectory(); },
+				.SetValue = [this](FString Value)	{ GetBuildCookRun()->SetPackageDirectory(Value); },
+				.GetDefaultValue = [this]()			{ return GetDefaultBuildCookRun()->GetPackageDirectory(); },
 			}
 		)
 		.AddBoolean( LOCTEXT("IoStoreLabel", "Use container files for optimized loading (I/O Store)"),
 			{
-				.GetValue = [this]()				{ return GetProfile()->IsUsingIoStore(); },
-				.SetValue = [this](bool bValue)		{ GetProfile()->SetUseIoStore(bValue); },
-				.GetDefaultValue = [this]()			{ return GetDefaultProfile()->IsUsingIoStore(); },
+				.GetValue = [this]()				{ return GetBuildCookRun()->IsUsingIoStore(); },
+				.SetValue = [this](bool bValue)		{ GetBuildCookRun()->SetUseIoStore(bValue); },
+				.GetDefaultValue = [this]()			{ return GetDefaultBuildCookRun()->IsUsingIoStore(); },
 			}
 		)
 		.AddFileString( LOCTEXT("IoStoreRefContainerLabel", "Optional I/O Store Reference Chunk Database"),
 			{
-				.GetValue = [this]()				{ return GetProfile()->GetReferenceContainerGlobalFileName(); },
-				.SetValue = [this](FString Value)	{ GetProfile()->SetReferenceContainerGlobalFileName(Value); },
-				.GetDefaultValue = [this]()			{ return GetDefaultProfile()->GetReferenceContainerGlobalFileName(); },
-				.IsEnabled = [this]()				{ return GetProfile()->IsUsingIoStore(); },
+				.GetValue = [this]()				{ return GetBuildCookRun()->GetReferenceContainerGlobalFileName(); },
+				.SetValue = [this](FString Value)	{ GetBuildCookRun()->SetReferenceContainerGlobalFileName(Value); },
+				.GetDefaultValue = [this]()			{ return GetDefaultBuildCookRun()->GetReferenceContainerGlobalFileName(); },
+				.IsEnabled = [this]()				{ return GetBuildCookRun()->IsUsingIoStore(); },
 			},
 			TEXT("global.utoc files|global.utoc")
 		)
 		.AddFileString( LOCTEXT("IoStoreRefCryptoKeys", "Optional I/O Store Reference Chunk Crypto Key"),
 			{
-				.GetValue = [this]()				{ return GetProfile()->GetReferenceContainerCryptoKeysFileName(); },
-				.SetValue = [this](FString Value)	{ GetProfile()->SetReferenceContainerCryptoKeysFileName(Value); },
-				.GetDefaultValue = [this]()			{ return GetDefaultProfile()->GetReferenceContainerCryptoKeysFileName(); },
-				.IsEnabled = [this]()				{ return GetProfile()->IsUsingIoStore(); },
+				.GetValue = [this]()				{ return GetBuildCookRun()->GetReferenceContainerCryptoKeysFileName(); },
+				.SetValue = [this](FString Value)	{ GetBuildCookRun()->SetReferenceContainerCryptoKeysFileName(Value); },
+				.GetDefaultValue = [this]()			{ return GetDefaultBuildCookRun()->GetReferenceContainerCryptoKeysFileName(); },
+				.IsEnabled = [this]()				{ return GetBuildCookRun()->IsUsingIoStore(); },
 			},
 			TEXT("crypto.json files|crypto.json")
 		)
 	;
 }
 
-bool FDeprecatedPropertiesLaunchExtensionInstance::CanBeRemoved() const
-{
-	// cannot remove the extension while there are deprecated properties set
-	return !Owner->HasDeprecatedProperties(GetProfile(), GetDefaultProfile());
-}
-
 void FDeprecatedPropertiesLaunchExtension::MakeCustomExtensionSubmenu(FMenuBuilder& MenuBuilder, ILauncherProfileRef InProfile, TSharedRef<ProjectLauncher::FModel> InModel)
 {
 	// putting this here for now because it isn't serialized with the profile. @todo: should have a better home!
@@ -170,20 +190,33 @@ FText FDeprecatedPropertiesLaunchExtension::GetDisplayName() const
 
 bool FDeprecatedPropertiesLaunchExtension::IsCreatedByDefault( ILauncherProfileRef InProfile, TSharedRef<ProjectLauncher::FModel> InModel ) const
 {
-	ILauncherProfileRef DefaultProfile = InModel->GetDefaultCustomLaunchProfile();
-	return HasDeprecatedProperties(InProfile, DefaultProfile);
+	return HasDeprecatedProperties(InProfile, InModel);
 }
 
-bool FDeprecatedPropertiesLaunchExtension::HasDeprecatedProperties(ILauncherProfileRef InProfile, ILauncherProfileRef InDefaultProfile) const
+bool FDeprecatedPropertiesLaunchExtension::HasDeprecatedProperties(const ILauncherProfileBuildCookRunRef& InBuildCookRun, TSharedRef<ProjectLauncher::FModel> InModel) const
 {
-	return	InProfile->GetSkipCookingEditorContent()				!=	InDefaultProfile->GetSkipCookingEditorContent() ||
-			InProfile->GetCookConfiguration()						!=	InDefaultProfile->GetCookConfiguration() ||
-			InProfile->GetPackageDirectory()						!=	InDefaultProfile->GetPackageDirectory() ||
-			InProfile->IsUsingIoStore()								!=	InDefaultProfile->IsUsingIoStore() ||
-			InProfile->GetReferenceContainerGlobalFileName()		!=	InDefaultProfile->GetReferenceContainerGlobalFileName() ||
-			InProfile->GetReferenceContainerCryptoKeysFileName()	!=	InDefaultProfile->GetReferenceContainerCryptoKeysFileName()
+	const ILauncherProfileBuildCookRunRef& DefaultBuildCookRun = InModel->GetCustomDefaultBuildCookRun();
+
+	return	InBuildCookRun->GetSkipCookingEditorContent()				!=	DefaultBuildCookRun->GetSkipCookingEditorContent() ||
+			InBuildCookRun->GetCookConfiguration()						!=	DefaultBuildCookRun->GetCookConfiguration() ||
+			InBuildCookRun->GetPackageDirectory()						!=	DefaultBuildCookRun->GetPackageDirectory() ||
+			InBuildCookRun->IsUsingIoStore()							!=	DefaultBuildCookRun->IsUsingIoStore() ||
+			InBuildCookRun->GetReferenceContainerGlobalFileName()		!=	DefaultBuildCookRun->GetReferenceContainerGlobalFileName() ||
+			InBuildCookRun->GetReferenceContainerCryptoKeysFileName()	!=	DefaultBuildCookRun->GetReferenceContainerCryptoKeysFileName()
 		;
+}
 
+bool FDeprecatedPropertiesLaunchExtension::HasDeprecatedProperties(ILauncherProfileRef InProfile, TSharedRef<ProjectLauncher::FModel> InModel) const
+{
+	for (const ILauncherProfileBuildCookRunRef& BuildCookRun : InProfile->GetBuildCookRunCommands())
+	{
+		if (HasDeprecatedProperties(BuildCookRun, InModel))
+		{
+			return true;
+		}
+	}
+
+	return false;
 }
 
 
diff --git a/Engine/Plugins/Developer/ProjectLauncher/Source/CommonLaunchExtensions/Private/DeprecatedProperties/DeprecatedPropertiesLaunchExtension.h b/Engine/Plugins/Developer/ProjectLauncher/Source/CommonLaunchExtensions/Private/DeprecatedProperties/DeprecatedPropertiesLaunchExtension.h
--- a/Engine/Plugins/Developer/ProjectLauncher/Source/CommonLaunchExtensions/Private/DeprecatedProperties/DeprecatedPropertiesLaunchExtension.h
+++ b/Engine/Plugins/Developer/ProjectLauncher/Source/CommonLaunchExtensions/Private/DeprecatedProperties/DeprecatedPropertiesLaunchExtension.h
@@ -2,24 +2,33 @@
 
 #pragma once
 
-#include "Extension/LaunchExtension.h"
+#include "Extension/BuildCookRunLaunchExtension.h"
 
 
-class FDeprecatedPropertiesLaunchExtensionInstance : public ProjectLauncher::FLaunchExtensionInstance
+class FDeprecatedPropertiesLaunchExtensionInstance : public ProjectLauncher::FBuildCookRunLaunchExtensionInstance
 {
 public:
 	FDeprecatedPropertiesLaunchExtensionInstance( FArgs& InArgs );
 	virtual ~FDeprecatedPropertiesLaunchExtensionInstance() = default;
 
-	virtual void CustomizeTree( ProjectLauncher::FLaunchProfileTreeData& ProfileTreeData ) override;
-	virtual bool CanBeRemoved() const override;
+	virtual TSharedRef<ProjectLauncher::FBuildCookRunExtension> CreateBuildCookRunExtension( const ProjectLauncher::FBuildCookRunExtension::FArgs& InArgs ) override;
+	virtual bool IsBuildCookRunExtensionEnabledByDefault( const ILauncherProfileBuildCookRunRef& InBuildCookRun ) const override;
+	virtual bool CanToggleBuildCookRunExtension( const ILauncherProfileBuildCookRunRef& InBuildCookRun, bool bWantToEnable ) const override;
 
 protected:
 	TSharedRef<class FDeprecatedPropertiesLaunchExtension> Owner;
+	class FBuildCookRunInstance : public ProjectLauncher::FBuildCookRunExtension
+	{
+	public:
+		FBuildCookRunInstance(const ProjectLauncher::FBuildCookRunExtension::FArgs& InArgs );
+		virtual void CustomizeTree( ProjectLauncher::FLaunchProfileTreeNode& ProfileTreeNode ) override;
+		TSharedRef<class FDeprecatedPropertiesLaunchExtension> Owner;
+	};
+
 };
 
 
-class FDeprecatedPropertiesLaunchExtension : public ProjectLauncher::FLaunchExtension
+class FDeprecatedPropertiesLaunchExtension : public ProjectLauncher::FBuildCookRunLaunchExtension
 {
 public:
 	virtual TSharedPtr<ProjectLauncher::FLaunchExtensionInstance> CreateInstanceForProfile( ProjectLauncher::FLaunchExtensionInstance::FArgs& InArgs ) override;
@@ -29,5 +38,6 @@ public:
 	virtual void GetExtensionsMenuEntry( FExtensionsMenuEntry& MenuEntry ) const override;
 	virtual void MakeCustomExtensionSubmenu(FMenuBuilder& MenuBuilder, ILauncherProfileRef InProfile, TSharedRef<ProjectLauncher::FModel> InModel) override;
 	
-	bool HasDeprecatedProperties(ILauncherProfileRef InProfile, ILauncherProfileRef InDefaultProfile) const;
+	bool HasDeprecatedProperties(const ILauncherProfileBuildCookRunRef& InBuildCookRun, TSharedRef<ProjectLauncher::FModel> InModel) const;
+	bool HasDeprecatedProperties(ILauncherProfileRef InProfile, TSharedRef<ProjectLauncher::FModel> InModel) const;
 };
\ No newline at end of file
diff --git a/Engine/Plugins/Developer/ProjectLauncher/Source/ProjectLauncher/Private/Extension/AutomatedTestLaunchExtension.cpp b/Engine/Plugins/Developer/ProjectLauncher/Source/ProjectLauncher/Private/Extension/AutomatedTestLaunchExtension.cpp
--- a/Engine/Plugins/Developer/ProjectLauncher/Source/ProjectLauncher/Private/Extension/AutomatedTestLaunchExtension.cpp
+++ b/Engine/Plugins/Developer/ProjectLauncher/Source/ProjectLauncher/Private/Extension/AutomatedTestLaunchExtension.cpp
@@ -13,27 +13,54 @@ namespace ProjectLauncher
 	}
 
 	void FAutomatedTestLaunchExtensionInstance::InternalInitialize()
+	{
+	}
+
+	void FAutomatedTestLaunchExtensionInstance::OnCustomValidation()
 	{
 		ILauncherProfileAutomatedTestPtr AutomatedTest = GetTest();
-		if (AutomatedTest.IsValid())
+		if (AutomatedTest.IsValid() && AutomatedTest->IsEnabled())
 		{
-			OnTestAdded(AutomatedTest.ToSharedRef());
+			OnValidateProfile();
+		}
+	}
 
-			int32 Order = GetTestOrder();
-			check(Order > 0);
-			AutomatedTest->SetOrder(Order);
+	bool FAutomatedTestLaunchExtensionInstance::ManagesUATCommand( const ILauncherProfileUATCommandRef& UATCommand ) const
+	{
+		return UATCommand == GetTest();
+	}
+
+
+
+	FLaunchProfileTreeNode& FAutomatedTestLaunchExtensionInstance::AddDefaultHeadingInternal(ProjectLauncher::FLaunchProfileTreeData& ProfileTreeData, FText DisplayName, int32 SortOrder) const
+	{
+		ILauncherProfileAutomatedTestPtr AutomatedTest = GetTest();
+		if (AutomatedTest.IsValid())
+		{
+			return ProfileTreeData.FindOrAddHeading(AutomatedTest.ToSharedRef(), true );
 		}
+
+		return FLaunchExtensionInstance::AddDefaultHeadingInternal(ProfileTreeData, DisplayName, SortOrder);
 	}
 
 
 	void FAutomatedTestLaunchExtensionInstance::OnAdded()
 	{
-		ILauncherProfileAutomatedTestRef AutomatedTest = GetProfile()->FindOrAddAutomatedTest(*GetTestInternalName());
-		OnTestAdded(AutomatedTest);
+		bool bIsNew = false;
+		ILauncherProfileAutomatedTestPtr AutomatedTest = GetTest();
+		if (!AutomatedTest.IsValid())
+		{
+			AutomatedTest = GetProfile()->FindOrAddAutomatedTest(*GetTestInternalName());
+			bIsNew = true;
+		}
+		OnTestAdded(AutomatedTest.ToSharedRef());
 
-		int32 Order = GetTestOrder();
-		check(Order > 0);
-		AutomatedTest->SetOrder(Order);
+		if (bIsNew)
+		{
+			int32 Order = GetTestOrder();
+			AutomatedTest->SetOrder(Order);
+			AutomatedTest->SetDescription( GetExtension()->GetDisplayName().ToString() );
+		}
 
 		BroadcastPropertyChanged();
 
@@ -48,6 +75,7 @@ namespace ProjectLauncher
 			GetProfile()->RemoveAutomatedTest(*GetTestInternalName());
 		}
 		BroadcastPropertyChanged();
+		RequestFullTreeRebuild(); // @fixme: still needed?
 	}
 
 	bool FAutomatedTestLaunchExtensionInstance::IsTestActive() const
diff --git a/Engine/Plugins/Developer/ProjectLauncher/Source/ProjectLauncher/Private/Extension/BuildCookRunLaunchExtension.cpp b/Engine/Plugins/Developer/ProjectLauncher/Source/ProjectLauncher/Private/Extension/BuildCookRunLaunchExtension.cpp
new file mode 100644
--- /dev/null
+++ b/Engine/Plugins/Developer/ProjectLauncher/Source/ProjectLauncher/Private/Extension/BuildCookRunLaunchExtension.cpp
@@ -0,0 +1,256 @@
+// Copyright Epic Games, Inc. All Rights Reserved.
+
+#include "Extension/BuildCookRunLaunchExtension.h"
+#include "Model/ProjectLauncherModel.h"
+
+#define LOCTEXT_NAMESPACE "FBuildCookRunLaunchExtensionInstance"
+
+
+namespace ProjectLauncher
+{
+	FBuildCookRunExtension::FBuildCookRunExtension( const FArgs& InArgs )
+		: Profile(InArgs.Profile)
+		, BuildCookRun(InArgs.BuildCookRun)
+		, DefaultBuildCookRun(InArgs.DefaultBuildCookRun)
+		, Model(InArgs.Model)
+		, Extension(InArgs.Extension)
+		, ExtensionInstance(InArgs.ExtensionInstance)
+		, OwnerTreeData(InArgs.OwnerTreeData)
+	{
+		Profile->OnCustomUATCommandValidation().AddRaw( this, &FBuildCookRunExtension::OnUATCommandValidation );
+		Profile->OnProjectChanged().AddRaw( this, &FBuildCookRunExtension::OnProjectChanged);
+	}
+
+	FBuildCookRunExtension::~FBuildCookRunExtension()
+	{
+		Profile->OnCustomUATCommandValidation().RemoveAll(this);
+		Profile->OnProjectChanged().RemoveAll(this);
+	}
+
+	void FBuildCookRunExtension::OnUATCommandValidation( ILauncherProfileUATCommandRef UATCommand )
+	{
+		if (UATCommand == BuildCookRun && UATCommand->IsEnabled())
+		{
+			OnValidate();
+		}
+	}
+
+	FLaunchProfileTreeNode& FBuildCookRunExtension::AddDefaultHeading(FLaunchProfileTreeNode& ProfileTreeData) const
+	{
+		return AddDefaultHeading(ProfileTreeData, Extension->GetDisplayName() );
+	}
+
+	FLaunchProfileTreeNode& FBuildCookRunExtension::AddDefaultHeading(FLaunchProfileTreeNode& ProfileTreeData, FText OverrideDisplayName) const
+	{
+		if (Model->CanUseSimplifiedLayout(Profile))
+		{
+			return ProfileTreeData.GetTreeData()->AddHeading(Extension->GetInternalName(), OverrideDisplayName, GetHeadingSortOrder() );
+		}
+
+		return ProfileTreeData.AddSubHeading(Extension->GetInternalName(), OverrideDisplayName, GetHeadingSortOrder() );
+	}
+
+
+
+
+	FBuildCookRunLaunchExtensionInstance::FBuildCookRunLaunchExtensionInstance( FArgs& InArgs )
+		: FLaunchExtensionInstance(InArgs)
+	{
+		GetProfile()->OnUATCommandAdded().AddRaw( this, &FBuildCookRunLaunchExtensionInstance::OnUATCommandAdded);
+		GetProfile()->OnUATCommandRemoved().AddRaw( this, &FBuildCookRunLaunchExtensionInstance::OnUATCommandRemoved);
+
+	}
+
+	FBuildCookRunLaunchExtensionInstance::~FBuildCookRunLaunchExtensionInstance()
+	{
+		BuildCookRunExtensions.Reset();
+		GetProfile()->OnUATCommandAdded().RemoveAll(this );
+		GetProfile()->OnUATCommandRemoved().RemoveAll(this );
+	}
+
+	void FBuildCookRunLaunchExtensionInstance::InternalInitialize()
+	{
+		FLaunchExtensionInstance::InternalInitialize();
+
+		// create all build cook run extensions if we support them		
+		for (const ILauncherProfileBuildCookRunRef& BuildCookRun : GetProfile()->GetBuildCookRunCommands())
+		{
+			bool bEnabledByDefault = IsBuildCookRunExtensionEnabledByDefault(BuildCookRun);
+			if (GetConfigBool(EConfig::PerProfile, *GetConfigKeyName(TEXT("Enabled"), BuildCookRun), bEnabledByDefault))
+			{
+				CreateBuildCookRunExtensionInternal(BuildCookRun);
+			}
+		}
+	}
+
+	void FBuildCookRunLaunchExtensionInstance::CustomizeTree( FLaunchProfileTreeData& ProfileTreeData )
+	{
+		FLaunchExtensionInstance::CustomizeTree(ProfileTreeData);
+
+		for (const TTuple<ILauncherProfileUATCommandPtr,TSharedRef<FBuildCookRunExtension>>& Pair : BuildCookRunExtensions)
+		{
+			const bool bIsPrimaryHeading = true;
+			FLaunchProfileTreeNode& ProfileTreeNode = ProfileTreeData.FindOrAddHeading(Pair.Key.ToSharedRef(), bIsPrimaryHeading);
+
+			Pair.Value->CustomizeTree(ProfileTreeNode);
+		}
+	}
+
+	FString FBuildCookRunLaunchExtensionInstance::GetConfigKeyName( const TCHAR* Name, const ILauncherProfileBuildCookRunRef& BuildCookRun) const
+	{
+		return FString::Printf(TEXT("%s.%s"), BuildCookRun->GetInternalName(), Name);
+	}
+
+	bool FBuildCookRunLaunchExtensionInstance::HasActiveBuildCookRunExtension() const
+	{
+		return (BuildCookRunExtensions.Num() > 0);
+	}
+
+	TSharedRef<FBuildCookRunExtension> FBuildCookRunLaunchExtensionInstance::CreateBuildCookRunExtension( const FBuildCookRunExtension::FArgs& InArgs )
+	{
+		checkNoEntry();
+		TSharedPtr<FBuildCookRunExtension> Invalid;
+		return Invalid.ToSharedRef();
+
+	}
+	
+	TSharedRef<FBuildCookRunExtension> FBuildCookRunLaunchExtensionInstance::CreateBuildCookRunExtensionInternal( const ILauncherProfileBuildCookRunRef& InBuildCookRun )
+	{
+		TSharedRef<FBuildCookRunExtension>* ExistingInstancePtr = BuildCookRunExtensions.Find(InBuildCookRun.ToSharedPtr());
+		if (ExistingInstancePtr != nullptr)
+		{
+			return (*ExistingInstancePtr);
+		}
+
+		FBuildCookRunExtension::FArgs Args
+		{
+			.Profile = GetProfile(),
+			.BuildCookRun = InBuildCookRun,
+			.DefaultBuildCookRun = GetModel()->GetCustomDefaultBuildCookRun(),
+			.Model = GetModel(),
+			.Extension = GetExtension(),
+			.ExtensionInstance = AsShared(),
+			.OwnerTreeData = GetOwnerTreeData(),
+		};
+
+		TSharedRef<FBuildCookRunExtension> NewInstance = CreateBuildCookRunExtension(Args);
+
+		BuildCookRunExtensions.Add(InBuildCookRun.ToSharedPtr(), NewInstance);
+
+		return NewInstance;
+	}
+
+	void FBuildCookRunLaunchExtensionInstance::CustomizeBuildCookRunExtensionMenu( FMenuBuilder& InMenuBuilder, const ILauncherProfileBuildCookRunRef& InBuildCookRun )
+	{
+		auto IsEnabled = [this, InBuildCookRun]()
+		{
+			return BuildCookRunExtensions.Contains(InBuildCookRun.ToSharedPtr());
+		};
+
+		auto ToggleEnabled = [this, InBuildCookRun, IsEnabled]()
+		{
+			bool bEnabled = !IsEnabled();
+			SetConfigBool(EConfig::PerProfile, *GetConfigKeyName(TEXT("Enabled"), InBuildCookRun), bEnabled);
+			RequestFullTreeRebuild();
+		};
+
+		auto CanToggle = [this, InBuildCookRun, IsEnabled]
+		{
+			bool bWantToEnable = !IsEnabled();
+			return CanToggleBuildCookRunExtension(InBuildCookRun, bWantToEnable);
+		};
+
+		auto GetToolTipText = [InBuildCookRun, CanToggle, IsEnabled]
+		{
+			bool bIsEnabled = IsEnabled();
+			bool bCanToggle = CanToggle();
+
+			if (bIsEnabled && !bCanToggle)
+			{
+				return LOCTEXT("CannotRemoveExtensionTooltip", "The extension cannot be removed at this time");
+			}
+			else if (!bIsEnabled && !bCanToggle)
+			{
+				return LOCTEXT("CannotAddExtensionTooltip", "The extension cannot be added at this time");
+			}
+
+			return FText::GetEmpty();
+		};
+
+
+		InMenuBuilder.AddMenuEntry(
+			GetExtension()->GetDisplayName(),
+			GetToolTipText(),
+			FSlateIcon(),
+			FUIAction(
+				FExecuteAction::CreateLambda(ToggleEnabled),
+				FCanExecuteAction::CreateLambda(CanToggle),
+				FIsActionChecked::CreateLambda(IsEnabled)
+			),
+			NAME_None,
+			EUserInterfaceActionType::Check);
+	}
+
+	bool FBuildCookRunLaunchExtensionInstance::IsBuildCookRunExtensionEnabledByDefault( const ILauncherProfileBuildCookRunRef& InBuildCookRun ) const
+	{
+		return false;
+	}
+
+	bool FBuildCookRunLaunchExtensionInstance::CanToggleBuildCookRunExtension( const ILauncherProfileBuildCookRunRef& InBuildCookRun, bool bWantToEnable ) const
+	{
+		return true;
+	}
+
+
+
+	void FBuildCookRunLaunchExtensionInstance::OnUATCommandAdded( const ILauncherProfileUATCommandRef& UATCommand )
+	{
+		ILauncherProfileBuildCookRunPtr BuildCookRun = UATCommand->AsBuildCookRun();
+		if (BuildCookRun.IsValid())
+		{
+			if (!BuildCookRunExtensions.Contains(UATCommand.ToSharedPtr()) && IsBuildCookRunExtensionEnabledByDefault(BuildCookRun.ToSharedRef()))
+			{
+				CreateBuildCookRunExtensionInternal(BuildCookRun.ToSharedRef());
+				RequestFullTreeRebuild();
+			}
+		}
+	}
+
+	void FBuildCookRunLaunchExtensionInstance::OnUATCommandRemoved(const ILauncherProfileUATCommandRef& UATCommand)
+	{
+		ILauncherProfileBuildCookRunPtr BuildCookRun = UATCommand->AsBuildCookRun();
+		if (BuildCookRun.IsValid())
+		{
+			if (BuildCookRunExtensions.Contains(UATCommand.ToSharedPtr()))
+			{
+				BuildCookRunExtensions.Remove(UATCommand.ToSharedPtr());
+				RequestFullTreeRebuild();
+			}
+		}
+	}
+
+	bool FBuildCookRunLaunchExtensionInstance::CustomizeUATCommandLine( const ILauncherProfileUATCommandRef& InUATCommand, FString& InOutCommandLine )
+	{
+		bool bResult = FLaunchExtensionInstance::CustomizeUATCommandLine(InUATCommand, InOutCommandLine);
+
+		TSharedRef<FBuildCookRunExtension>* ExistingInstancePtr = BuildCookRunExtensions.Find(InUATCommand.ToSharedPtr());
+		if (ExistingInstancePtr != nullptr)
+		{
+			(*ExistingInstancePtr)->CustomizeUATCommandLine(InOutCommandLine);
+			return true;
+		}
+
+		return bResult;
+	}
+
+	void FBuildCookRunLaunchExtensionInstance::OnPropertyChanged() 
+	{
+		for (const TTuple<ILauncherProfileUATCommandPtr,TSharedRef<FBuildCookRunExtension>>& Pair : BuildCookRunExtensions)
+		{
+			Pair.Value->OnPropertyChanged();
+		}
+	}
+
+}
+
+#undef LOCTEXT_NAMESPACE
diff --git a/Engine/Plugins/Developer/ProjectLauncher/Source/ProjectLauncher/Private/Extension/LaunchExtension.cpp b/Engine/Plugins/Developer/ProjectLauncher/Source/ProjectLauncher/Private/Extension/LaunchExtension.cpp
--- a/Engine/Plugins/Developer/ProjectLauncher/Source/ProjectLauncher/Private/Extension/LaunchExtension.cpp
+++ b/Engine/Plugins/Developer/ProjectLauncher/Source/ProjectLauncher/Private/Extension/LaunchExtension.cpp
@@ -324,6 +324,15 @@ namespace ProjectLauncher
 		}
 	}
 
+	void FLaunchExtension::DestroyExtensionInstancesForProfile(ILauncherProfilePtr InProfile)
+	{
+		if (!GExtensionsDestroyed && InProfile.IsValid())
+		{
+			RemoveExtensionInstancesForProfile(InProfile.ToSharedRef());
+		}
+	}
+
+
 	TSharedPtr<FLaunchExtensionInstance> FLaunchExtension::GetExtensionInstance(TSharedRef<FLaunchExtension> InExtension, ILauncherProfileRef InProfile)
 	{
 		check(!GExtensionsDestroyed);
@@ -413,9 +422,6 @@ namespace ProjectLauncher
 
 
 
-
-
-
 	FLaunchExtensionInstance::FLaunchExtensionInstance( FArgs& InArgs )
 		: Profile(InArgs.Profile)
 		, DefaultProfile(InArgs.DefaultProfile)
@@ -423,7 +429,7 @@ namespace ProjectLauncher
 		, Extension(InArgs.Extension)
 		, OwnerTreeData(InArgs.OwnerTreeData)
 	{
-		Profile->OnCustomValidation().AddRaw( this, &FLaunchExtensionInstance::OnValidateProfile);
+		Profile->OnCustomValidation().AddRaw( this, &FLaunchExtensionInstance::OnCustomValidation);
 		Profile->OnProjectChanged().AddRaw( this, &FLaunchExtensionInstance::OnProjectChanged);
 	}
 
@@ -433,6 +439,20 @@ namespace ProjectLauncher
 		Profile->OnProjectChanged().RemoveAll(this);
 	}
 
+	void FLaunchExtensionInstance::InternalInitialize()
+	{
+	}
+
+	void FLaunchExtensionInstance::CustomizeTree()
+	{
+		CustomizeTree(GetOwnerTreeData());
+	}
+
+
+	void FLaunchExtensionInstance::CustomizeTree( FLaunchProfileTreeData& ProfileTreeData )
+	{
+	}
+
 	bool FLaunchExtensionInstance::GetExtensionVariables( TArray<FString>& OutVariables ) const
 	{
 		return false;
@@ -562,7 +582,16 @@ namespace ProjectLauncher
 
 	FName FLaunchExtensionInstance::GetLaunchPlatformName() const
 	{
-		TArray<const PlatformInfo::FTargetPlatformInfo*> Infos = GetModel()->GetPlatformInfos(GetProfile());
+		TArray<const PlatformInfo::FTargetPlatformInfo*> Infos;
+
+		for (const ILauncherProfileBuildCookRunRef& BuildCookRun : Profile->GetBuildCookRunCommands())
+		{
+			for (const PlatformInfo::FTargetPlatformInfo* Info : GetModel()->GetPlatformInfos(BuildCookRun))
+			{
+				Infos.AddUnique(Info);
+			}
+		}
+
 		if (Infos.Num() == 0)
 		{
 			return NAME_None;
@@ -754,6 +783,11 @@ namespace ProjectLauncher
 		OwnerTreeData.RequestTreeRefresh();
 	}
 
+	void FLaunchExtensionInstance::RequestFullTreeRebuild()
+	{
+		OwnerTreeData.RequestFullTreeRebuild();
+	}
+
 	FLaunchProfileTreeNode& FLaunchExtensionInstance::AddDefaultHeading(ProjectLauncher::FLaunchProfileTreeData& ProfileTreeData) const
 	{
 		return AddDefaultHeading(ProfileTreeData, Extension->GetDisplayName());
@@ -761,8 +795,58 @@ namespace ProjectLauncher
 
 	FLaunchProfileTreeNode& FLaunchExtensionInstance::AddDefaultHeading(ProjectLauncher::FLaunchProfileTreeData& ProfileTreeData, FText OverrideDisplayName) const
 	{
-		return ProfileTreeData.AddHeading(Extension->GetInternalName(), OverrideDisplayName, GetHeadingSortOrder() );
+		return AddDefaultHeadingInternal(ProfileTreeData, OverrideDisplayName, GetHeadingSortOrder() );
+	}
+
+	FLaunchProfileTreeNode& FLaunchExtensionInstance::AddDefaultHeadingInternal(ProjectLauncher::FLaunchProfileTreeData& ProfileTreeData, FText DisplayName, int32 SortOrder) const
+	{
+		return ProfileTreeData.AddHeading(Extension->GetInternalName(), DisplayName, SortOrder );
+	}
+
+	bool FLaunchExtensionInstance::CustomizeUATCommandLine( const ILauncherProfileUATCommandRef& InUATCommand, FString& InOutCommandLine )
+	{
+		return false;
+	}
+
+	void FLaunchExtensionInstance::OnPropertyChanged() 
+	{
+	}
+
+	void FLaunchExtensionInstance::OnCustomValidation()
+	{
+		OnValidateProfile();
 	}
+
+
+
+
+
+
+
+
+
+
+
+
+
+
+
+
+
+
+
+
+
+
+
+
+
+
+
+
+
+
+
 }
 
 #undef LOCTEXT_NAMESPACE
diff --git a/Engine/Plugins/Developer/ProjectLauncher/Source/ProjectLauncher/Private/Extension/UATCommandLaunchExtension.cpp b/Engine/Plugins/Developer/ProjectLauncher/Source/ProjectLauncher/Private/Extension/UATCommandLaunchExtension.cpp
--- a/Engine/Plugins/Developer/ProjectLauncher/Source/ProjectLauncher/Private/Extension/UATCommandLaunchExtension.cpp
+++ b/Engine/Plugins/Developer/ProjectLauncher/Source/ProjectLauncher/Private/Extension/UATCommandLaunchExtension.cpp
@@ -13,20 +13,55 @@ namespace ProjectLauncher
 	}
 
 	void FUATCommandLaunchExtensionInstance::InternalInitialize()
+	{
+		FLaunchExtensionInstance::InternalInitialize();
+	}
+
+	void FUATCommandLaunchExtensionInstance::OnCustomValidation()
+	{
+		ILauncherProfileUATCommandPtr UATCommand = GetUATCommand();
+		if (UATCommand.IsValid() && UATCommand->IsEnabled())
+		{
+			OnValidateProfile();
+		}
+	}
+
+	bool FUATCommandLaunchExtensionInstance::ManagesUATCommand( const ILauncherProfileUATCommandRef& UATCommand ) const
+	{
+		return UATCommand == GetUATCommand();
+	}
+
+
+
+	FLaunchProfileTreeNode& FUATCommandLaunchExtensionInstance::AddDefaultHeadingInternal(ProjectLauncher::FLaunchProfileTreeData& ProfileTreeData, FText DisplayName, int32 SortOrder) const
 	{
 		ILauncherProfileUATCommandPtr UATCommand = GetUATCommand();
 		if (UATCommand.IsValid())
 		{
-			OnUATCommandAdded(UATCommand.ToSharedRef());
-			UATCommand->SetOrder(GetUATCommandOrder());
+			return ProfileTreeData.FindOrAddHeading(UATCommand.ToSharedRef(), true );
 		}
+
+		return FLaunchExtensionInstance::AddDefaultHeadingInternal(ProfileTreeData, DisplayName, SortOrder);
 	}
 
 	void FUATCommandLaunchExtensionInstance::OnAdded()
 	{
-		ILauncherProfileUATCommandRef UATCommand = GetProfile()->FindOrAddUATCommand(*GetUATCommandInternalName());
-		OnUATCommandAdded(UATCommand);
-		UATCommand->SetOrder(GetUATCommandOrder());
+		bool bIsNew = false;
+		ILauncherProfileUATCommandPtr UATCommand = GetUATCommand();
+		if (!UATCommand.IsValid())
+		{
+			UATCommand = GetProfile()->FindOrAddUATCommand(*GetUATCommandInternalName());
+			bIsNew = true;
+		}
+		OnUATCommandAdded(UATCommand.ToSharedRef());
+
+		if (bIsNew)
+		{
+			int32 Order = GetUATCommandOrder();
+			UATCommand->SetOrder(Order);
+			UATCommand->SetDescription( GetExtension()->GetDisplayName().ToString() );
+		}
+
 		BroadcastPropertyChanged();
 	}
 	void FUATCommandLaunchExtensionInstance::OnRemoved()
diff --git a/Engine/Plugins/Developer/ProjectLauncher/Source/ProjectLauncher/Private/Model/ProjectLauncherModel.cpp b/Engine/Plugins/Developer/ProjectLauncher/Source/ProjectLauncher/Private/Model/ProjectLauncherModel.cpp
--- a/Engine/Plugins/Developer/ProjectLauncher/Source/ProjectLauncher/Private/Model/ProjectLauncherModel.cpp
+++ b/Engine/Plugins/Developer/ProjectLauncher/Source/ProjectLauncher/Private/Model/ProjectLauncherModel.cpp
@@ -28,12 +28,18 @@
 namespace ProjectLauncher
 {
 	const FText ProjectLauncherText = LOCTEXT("ProjectLauncherTitle", "Project Launcher");
+	bool bUseClassicView = true;
 
 	FModel::FModel(const TSharedRef<ITargetDeviceProxyManager>& InDeviceProxyManager, const TSharedRef<ILauncher>& InLauncher, const TSharedRef<ILauncherProfileManager>& InProfileManager)
 		: DeviceProxyManager(InDeviceProxyManager)
 		, Launcher(InLauncher)
 		, ProfileManager(InProfileManager)
 	{
+		// use a custom ini file so that it can be shared between UnrealFrontend and the editor
+		ConfigFileName = FPaths::EngineSavedDir() / TEXT("Config") / TEXT("ProjectLauncher") / TEXT("UserSettings.ini");
+		FConfigContext::ReadSingleIntoGConfig().Load(*ConfigFileName);
+		LoadConfig();
+
 		// register callbacks
 		ProfileManager->OnPostProcessLaunchCommandLine().AddRaw( this, &FModel::OnModifyLaunchCommandLine );
 		ProfileManager->OnPostProcessUATCommandLine().AddRaw( this, &FModel::OnModifyUATCommandLine );
@@ -44,7 +50,6 @@ namespace ProjectLauncher
 		DeviceProxyManager->OnProxyAdded().AddRaw(this, &FModel::HandleDeviceProxyAdded);
 		DeviceProxyManager->OnProxyRemoved().AddRaw(this, &FModel::HandleDeviceProxyRemoved);
 
-
 		// ensure there's a project when we're in the editor
 		// this means we don't need to display the global project selector in the editor
 		if (GIsEditor && FPaths::IsProjectFilePathSet())
@@ -52,10 +57,9 @@ namespace ProjectLauncher
 			ProfileManager->SetProjectPath(FPaths::GetProjectFilePath());
 		}
 
-
 		// prepare profiles
 		BasicLaunchProfile = CreateBasicLaunchProfile();
-		bHasSetBasicLaunchProfilePlatform = (BasicLaunchProfile->GetCookedPlatforms().Num() > 0);
+		bHasSetBasicLaunchProfilePlatform = (BasicLaunchProfile->GetDeployedDeviceGroup()->GetDeviceIDs().Num() > 0);
 
 		AllProfiles = ProfileManager->GetAllProfiles();
 		for (const ILauncherProfilePtr& Profile : AllProfiles)
@@ -70,11 +74,6 @@ namespace ProjectLauncher
 		DefaultCustomLaunchProfile = CreateCustomProfile(TEXT("DefaultCustomProfile"));
 
 
-		// use a custom ini file so that it can be shared between UnrealFrontend and the editor
-		ConfigFileName = FPaths::EngineSavedDir() / TEXT("Config") / TEXT("ProjectLauncher") / TEXT("UserSettings.ini");
-		FConfigContext::ReadSingleIntoGConfig().Load(*ConfigFileName);
-
-		LoadConfig();
 	}
 
 
@@ -133,11 +132,17 @@ namespace ProjectLauncher
 
 	bool FModel::IsAdvancedProfile(const ILauncherProfileRef& Profile) const
 	{
-		return 
-				Profile->GetPackagingMode() == ELauncherProfilePackagingModes::SharedRepository ||
-				Profile->GetDeploymentMode() == ELauncherProfileDeploymentModes::CopyRepository ||
-				Profile->GetLaunchMode() == ELauncherProfileLaunchModes::CustomRoles
-			;
+		for (const ILauncherProfileBuildCookRunRef& BuildCookRun : Profile->GetBuildCookRunCommands())
+		{
+			if (BuildCookRun->GetPackagingMode() == ELauncherProfilePackagingModes::SharedRepository ||
+				BuildCookRun->GetDeploymentMode() == ELauncherProfileDeploymentModes::CopyRepository ||
+				BuildCookRun->GetLaunchMode() == ELauncherProfileLaunchModes::CustomRoles)
+			{
+				return true;
+			}
+		}
+
+		return false;
 	}
 
 
@@ -201,17 +206,32 @@ namespace ProjectLauncher
 		return nullptr;
 	}
 
-	TArray<const PlatformInfo::FTargetPlatformInfo*> FModel::GetPlatformInfos(const ILauncherProfilePtr& Profile)
+	TArray<const PlatformInfo::FTargetPlatformInfo*> FModel::GetPlatformInfos(const ILauncherProfileBuildCookRunRef& BuildCookRun)
+	{
+		TArray<const PlatformInfo::FTargetPlatformInfo*> Result;
+		for (const FString& Platform : BuildCookRun->GetCookedPlatforms())
+		{
+			const PlatformInfo::FTargetPlatformInfo* PlatformInfo = PlatformInfo::FindPlatformInfo(FName(Platform));
+			if (PlatformInfo != nullptr)
+			{
+				Result.Add(PlatformInfo);
+			}
+		}
+
+		return MoveTemp(Result);
+	}
+
+	TArray<const PlatformInfo::FTargetPlatformInfo*> FModel::GetPlatformInfosForAllCommands(const ILauncherProfilePtr& Profile)
 	{
 		TArray<const PlatformInfo::FTargetPlatformInfo*> Result;
+
 		if (Profile.IsValid())
 		{
-			for (const FString& Platform : Profile->GetCookedPlatforms())
+			for (const ILauncherProfileBuildCookRunRef& BuildCookRun : Profile->GetBuildCookRunCommands())
 			{
-				const PlatformInfo::FTargetPlatformInfo* PlatformInfo = PlatformInfo::FindPlatformInfo(FName(Platform));
-				if (PlatformInfo != nullptr)
+				for (const PlatformInfo::FTargetPlatformInfo* PlatformInfo : GetPlatformInfos(BuildCookRun))
 				{
-					Result.Add(PlatformInfo);
+					Result.AddUnique(PlatformInfo);
 				}
 			}
 		}
@@ -279,16 +299,13 @@ namespace ProjectLauncher
 		return false;
 	}
 
-	bool FModel::IsUsingRemotePlatform(const ILauncherProfilePtr& Profile)
+	bool FModel::IsUsingRemotePlatform(const ILauncherProfileBuildCookRunRef& BuildCookRun)
 	{
-		if (Profile.IsValid())
+		for (const FString& Platform : BuildCookRun->GetCookedPlatforms())
 		{
-			for (const FString& Platform : Profile->GetCookedPlatforms())
+			if (!IsHostPlatform(FName(Platform)))
 			{
-				if (!IsHostPlatform(FName(Platform)))
-				{
-					return true;
-				}
+				return true;
 			}
 		}
 
@@ -330,10 +347,10 @@ PRAGMA_ENABLE_DEPRECATION_WARNINGS
 		return GetBuildTargetInfo(BuildTargetName, ProjectPath);
 	}
 
-	TArray<FTargetInfo> FModel::GetBuildTargetInfos(const ILauncherProfileRef& Profile)
+	TArray<FTargetInfo> FModel::GetBuildTargetInfos(const ILauncherProfileRef& Profile, const ILauncherProfileBuildCookRunRef& BuildCookRun)
 	{
 		const FString& ProjectPath = Profile->GetProjectPath();
-		const TArray<FString> BuildTargets = Profile->GetBuildTargets();
+		const TArray<FString> BuildTargets = BuildCookRun->GetBuildTargets();
 
 		TArray<FTargetInfo> Result;
 		if (BuildTargets.Num() == 0 || BuildTargets.Contains(FString()))
@@ -442,10 +459,10 @@ PRAGMA_ENABLE_DEPRECATION_WARNINGS
 
 	void FModel::HandleDeviceProxyAdded(const TSharedRef<ITargetDeviceProxy>& DeviceProxy )
 	{
-		if (!bHasSetBasicLaunchProfilePlatform)
+		if (!bHasSetBasicLaunchProfilePlatform && DeviceProxy->HasTargetPlatform(FPlatformProperties::IniPlatformName()) && ensure(BasicLaunchProfile->GetFirstBuildCookRun().IsValid()))
 		{
-			UpdatedCookedPlatformsFromDeployDeviceProxy(BasicLaunchProfile.ToSharedRef(), DeviceProxy);
-			bHasSetBasicLaunchProfilePlatform = true;
+			UpdatedCookedPlatformsFromDeployDeviceProxy(BasicLaunchProfile.ToSharedRef(), BasicLaunchProfile->GetFirstBuildCookRun().ToSharedRef(), DeviceProxy);
+			bHasSetBasicLaunchProfilePlatform = (BasicLaunchProfile->GetDeployedDeviceGroup()->GetDeviceIDs().Num() > 0);
 		}
 	}
 
@@ -456,9 +473,9 @@ PRAGMA_ENABLE_DEPRECATION_WARNINGS
 
 
 
-	TSharedPtr<ITargetDeviceProxy> FModel::GetDeviceProxy(const ILauncherProfileRef& Profile)
+	TSharedPtr<ITargetDeviceProxy> FModel::GetDeviceProxy(const ILauncherProfileBuildCookRunRef& BuildCookRun)
 	{
-		ILauncherDeviceGroupPtr DeployedDeviceGroup = Profile->GetDeployedDeviceGroup();
+		ILauncherDeviceGroupPtr DeployedDeviceGroup = BuildCookRun->GetDeployedDeviceGroup();
 		if (DeployedDeviceGroup != nullptr && DeployedDeviceGroup->GetDeviceIDs().Num() > 0)
 		{
 			const FString& DeviceID = DeployedDeviceGroup->GetDeviceIDs()[0];
@@ -477,36 +494,60 @@ PRAGMA_ENABLE_DEPRECATION_WARNINGS
 		//create the profile
 		ILauncherProfileRef Profile = ProfileManager->CreateUnsavedProfile(Name);
 
+		//set defaults
+		Profile->SetBuildUAT(false);
+		if (!ProfileManager->GetProjectPath().IsEmpty())
+		{
+			Profile->SetProjectPath(ProfileManager->GetProjectPath());
+			Profile->SetProjectSpecified(true);
+		}
+
+		//new profiles should always have a default BuildCookRun set up
+		check(Profile->GetFirstBuildCookRun().IsValid());
+		ILauncherProfileBuildCookRunRef BuildCookRun = Profile->GetFirstBuildCookRun().ToSharedRef();
+		SetDefaults(Profile, BuildCookRun);
+
+		return Profile;
+	}
+
+	void FModel::SetDefaults(const ILauncherProfileRef& Profile, const ILauncherProfileBuildCookRunRef& BuildCookRun)
+	{
 		// set defaults
 		FProjectSettings ProjectSettings = GetProjectSettings(Profile->GetProjectPath());
 		if (ProjectSettings.bUseZenStore && ProjectSettings.bAllowRemoteNetworkService)
 		{
-			SetProfileContentScheme(EContentScheme::ZenStreaming, Profile );
+			SetProfileContentScheme(EContentScheme::ZenStreaming, BuildCookRun );
 		}
 		else
 		{
-			SetProfileContentScheme(EContentScheme::PakFiles, Profile );
+			SetProfileContentScheme(EContentScheme::PakFiles, BuildCookRun );
 		}
-		Profile->SetBuildConfiguration(EBuildConfiguration::Development);
-		Profile->SetLaunchMode(ELauncherProfileLaunchModes::DefaultRole);
-		Profile->SetBuildMode(ELauncherProfileBuildModes::Auto);
-		Profile->SetBuildUAT(false);
+		BuildCookRun->SetBuildConfiguration(EBuildConfiguration::Development);
+		BuildCookRun->SetLaunchMode(ELauncherProfileLaunchModes::DefaultRole);
+		BuildCookRun->SetBuildMode(ELauncherProfileBuildModes::Auto);
 		// @fixme: set all defaults here 
 		// ...
 
 
 		// make sure there is a device & deploy group
-		ILauncherDeviceGroupRef DeployDeviceGroup = ProfileManager->AddNewDeviceGroup();
-		Profile->SetDeployedDeviceGroup(DeployDeviceGroup);
+		if (BuildCookRun->GetDeployedDeviceGroup() == nullptr)
+		{
+			ILauncherDeviceGroupRef DeployDeviceGroup = ProfileManager->AddNewDeviceGroup();
+			BuildCookRun->SetDeployedDeviceGroup(DeployDeviceGroup);
+		}
 
 		TArray<TSharedPtr<ITargetDeviceProxy>> DeviceProxies;
 		DeviceProxyManager->GetProxies(NAME_None, true, DeviceProxies);
-		if (DeviceProxies.Num() > 0 && DeviceProxies[0].IsValid())
+
+		if (DeviceProxies.Num() > 0 && DeviceProxies[0].IsValid()) // item 0 should always be the host platform
+		{
+			UpdatedCookedPlatformsFromDeployDeviceProxy(Profile, BuildCookRun, DeviceProxies[0].ToSharedRef());
+		}
+		else
 		{
-			UpdatedCookedPlatformsFromDeployDeviceProxy(Profile, DeviceProxies[0].ToSharedRef());
+			BuildCookRun->AddCookedPlatform(FPlatformProperties::IniPlatformName());
 		}
 
-		return Profile;
 	}
 
 	ILauncherProfileRef FModel::CreateBasicLaunchProfile()
@@ -525,71 +566,74 @@ PRAGMA_ENABLE_DEPRECATION_WARNINGS
 
 	void FModel::UpdateDeployDeviceProxyFromCookedPlatforms(const ILauncherProfileRef& Profile)
 	{
-		ILauncherDeviceGroupPtr DeviceGroup = Profile->GetDeployedDeviceGroup();
-		if (DeviceGroup.IsValid())
+		for (const ILauncherProfileBuildCookRunRef& BuildCookRun : Profile->GetBuildCookRunCommands())
 		{
-			TSet<FString> DeviceIDs( DeviceGroup->GetDeviceIDs() );
+			ILauncherDeviceGroupPtr DeviceGroup = BuildCookRun->GetDeployedDeviceGroup();
+			if (DeviceGroup.IsValid())
+			{
+				TSet<FString> DeviceIDs( DeviceGroup->GetDeviceIDs() );
 
-			TArray<FString> Platforms = Profile->GetCookedPlatforms();
+				TArray<FString> Platforms = BuildCookRun->GetCookedPlatforms();
 
-			TSet<FString> NewDeviceIDs;
-			for ( const FString& DeviceID : DeviceIDs)
-			{
-				TSharedPtr<ITargetDeviceProxy> DeviceProxy = DeviceProxyManager->FindProxyDeviceForTargetDevice(DeviceID);
-				if (DeviceProxy.IsValid())
+				TSet<FString> NewDeviceIDs;
+				for ( const FString& DeviceID : DeviceIDs)
 				{
-					TArray<FName> Variants;
-					DeviceProxy->GetVariants(Variants);
-
-					for (const FString& Platform : Platforms)
+					TSharedPtr<ITargetDeviceProxy> DeviceProxy = DeviceProxyManager->FindProxyDeviceForTargetDevice(DeviceID);
+					if (DeviceProxy.IsValid())
 					{
-						for( const FName& Variant : Variants)
+						TArray<FName> Variants;
+						DeviceProxy->GetVariants(Variants);
+
+						for (const FString& Platform : Platforms)
 						{
-							FString TargetPlatformName = DeviceProxy->GetTargetPlatformName(Variant);
-							if (Platform == TargetPlatformName)
+							for( const FName& Variant : Variants)
 							{
-								const TSet<FString>& TargetDeviceIds = DeviceProxy->GetTargetDeviceIds(Variant);
-								NewDeviceIDs.Append(TargetDeviceIds);
+								FString TargetPlatformName = DeviceProxy->GetTargetPlatformName(Variant);
+								if (Platform == TargetPlatformName)
+								{
+									const TSet<FString>& TargetDeviceIds = DeviceProxy->GetTargetDeviceIds(Variant);
+									NewDeviceIDs.Append(TargetDeviceIds);
+								}
 							}
 						}
 					}
+					else
+					{
+						NewDeviceIDs.Add(DeviceID);
+					}
 				}
-				else
-				{
-					NewDeviceIDs.Add(DeviceID);
-				}
-			}
 
-			// update the device group if there's been changes
-			if (NewDeviceIDs.Difference(DeviceIDs).Num() > 0 || DeviceIDs.Difference(NewDeviceIDs).Num() > 0)
-			{
-				DeviceGroup->RemoveAllDevices();
-
-				for (const FString& NewDeviceID : NewDeviceIDs)
+				// update the device group if there's been changes
+				if (NewDeviceIDs.Difference(DeviceIDs).Num() > 0 || DeviceIDs.Difference(NewDeviceIDs).Num() > 0)
 				{
-					DeviceGroup->AddDevice(NewDeviceID);
+					DeviceGroup->RemoveAllDevices();
+
+					for (const FString& NewDeviceID : NewDeviceIDs)
+					{
+						DeviceGroup->AddDevice(NewDeviceID);
+					}
 				}
 			}
 		}
 	}
 
-	void FModel::UpdatedCookedPlatformsFromDeployDeviceProxy(const ILauncherProfileRef& Profile, TSharedPtr<ITargetDeviceProxy> DeviceProxy)
+	void FModel::UpdatedCookedPlatformsFromDeployDeviceProxy(const ILauncherProfileRef& Profile, const ILauncherProfileBuildCookRunRef& BuildCookRun, TSharedPtr<ITargetDeviceProxy> DeviceProxy)
 	{
 		if (DeviceProxy.IsValid())
 		{
-			Profile->GetDeployedDeviceGroup()->RemoveAllDevices();
-			Profile->GetDeployedDeviceGroup()->AddDevice(DeviceProxy->GetTargetDeviceId(NAME_None));
+			BuildCookRun->GetDeployedDeviceGroup()->RemoveAllDevices();
+			BuildCookRun->GetDeployedDeviceGroup()->AddDevice(DeviceProxy->GetTargetDeviceId(NAME_None));
 		}
-		else if (Profile->GetDeployedDeviceGroup()->GetNumDevices() > 0)
+		else if (BuildCookRun->GetDeployedDeviceGroup()->GetNumDevices() > 0)
 		{
-			DeviceProxy = GetDeviceProxy(Profile);
+			DeviceProxy = GetDeviceProxy(BuildCookRun);
 		}
 		
 		if (DeviceProxy.IsValid())
 		{
-			Profile->ClearCookedPlatforms();
+			BuildCookRun->ClearCookedPlatforms();
 
-			TArray<FTargetInfo> BuildTargetInfos = GetBuildTargetInfos(Profile);
+			TArray<FTargetInfo> BuildTargetInfos = GetBuildTargetInfos(Profile, BuildCookRun);
 			for (const FTargetInfo& BuildTargetInfo : BuildTargetInfos)
 			{
 				FString PlatformName = DeviceProxy->GetTargetPlatformName(NAME_None);
@@ -597,19 +641,19 @@ PRAGMA_ENABLE_DEPRECATION_WARNINGS
 
 				if (PlatformInfo != nullptr)
 				{
-					Profile->AddCookedPlatform(PlatformInfo->Name.ToString());
+					BuildCookRun->AddCookedPlatform(PlatformInfo->Name.ToString());
 				}
 				else
 				{
-					Profile->AddCookedPlatform(DeviceProxy->GetTargetPlatformName(NAME_None));
+					BuildCookRun->AddCookedPlatform(DeviceProxy->GetTargetPlatformName(NAME_None));
 				}
 			}
 		}
 	}
 
-	void FModel::UpdateCookedPlatformsFromBuildTarget(const ILauncherProfileRef& Profile)
+	void FModel::UpdateCookedPlatformsFromBuildTarget(const ILauncherProfileRef& Profile, const ILauncherProfileBuildCookRunRef& BuildCookRun)
 	{
-		TArray<FTargetInfo> BuildTargetInfos = GetBuildTargetInfos(Profile);
+		TArray<FTargetInfo> BuildTargetInfos = GetBuildTargetInfos(Profile, BuildCookRun);
 
 		if (BuildTargetInfos.Num() != 1)
 		{
@@ -617,18 +661,18 @@ PRAGMA_ENABLE_DEPRECATION_WARNINGS
 		}
 		const FTargetInfo& BuildTargetInfo = BuildTargetInfos[0];
 
-		TArray<FString> Platforms = Profile->GetCookedPlatforms();
-		Profile->ClearCookedPlatforms();
+		TArray<FString> Platforms = BuildCookRun->GetCookedPlatforms();
+		BuildCookRun->ClearCookedPlatforms();
 		for (const FString& Platform : Platforms)
 		{
 			const PlatformInfo::FTargetPlatformInfo* PlatformInfo = GetPlatformInfo(*Platform, BuildTargetInfo);
 			if (PlatformInfo != nullptr)
 			{
-				Profile->AddCookedPlatform(PlatformInfo->Name.ToString());
+				BuildCookRun->AddCookedPlatform(PlatformInfo->Name.ToString());
 			}
 			else
 			{
-				Profile->AddCookedPlatform(Platform);
+				BuildCookRun->AddCookedPlatform(Platform);
 			}
 		}
 	}
@@ -671,15 +715,15 @@ PRAGMA_ENABLE_DEPRECATION_WARNINGS
 
 			
 
-	EContentScheme FModel::DetermineProfileContentScheme(const ILauncherProfileRef& Profile) const
+	EContentScheme FModel::DetermineProfileContentScheme(const ILauncherProfileBuildCookRunRef& BuildCookRun) const
 	{
-		if (Profile->GetCookMode() == ELauncherProfileCookModes::OnTheFly)
+		if (BuildCookRun->GetCookMode() == ELauncherProfileCookModes::OnTheFly)
 		{
 			return EContentScheme::CookOnTheFly;
 		}
-		else if (Profile->GetPackagingMode() != ELauncherProfilePackagingModes::DoNotPackage)
+		else if (BuildCookRun->GetPackagingMode() != ELauncherProfilePackagingModes::DoNotPackage)
 		{
-			if (Profile->IsCreatingReleaseVersion() || Profile->IsGeneratingPatch() || Profile->IsCreatingDLC())
+			if (BuildCookRun->IsCreatingReleaseVersion() || BuildCookRun->IsGeneratingPatch() || BuildCookRun->IsCreatingDLC())
 			{
 				return EContentScheme::SubmissionPackage;
 			}
@@ -688,15 +732,15 @@ PRAGMA_ENABLE_DEPRECATION_WARNINGS
 				return EContentScheme::DevelopmentPackage;
 			}
 		}
-		else if (Profile->IsPackingWithUnrealPak())
+		else if (BuildCookRun->IsPackingWithUnrealPak())
 		{
 			return EContentScheme::PakFiles;
 		}
-		else if (Profile->IsUsingZenPakStreaming())
+		else if (BuildCookRun->IsUsingZenPakStreaming())
 		{
 			return EContentScheme::ZenPakStreaming;
 		}
-		else if (Profile->IsUsingZenStore())
+		else if (BuildCookRun->IsUsingZenStore())
 		{
 			return EContentScheme::ZenStreaming;
 		}
@@ -708,7 +752,7 @@ PRAGMA_ENABLE_DEPRECATION_WARNINGS
 
 
 
-	void FModel::SetProfileContentScheme(EContentScheme ContentScheme, const ILauncherProfileRef& Profile, bool bWantToCook, ELauncherProfileDeploymentModes::Type DefaultDeploymentMode)
+	void FModel::SetProfileContentScheme(EContentScheme ContentScheme, const ILauncherProfileBuildCookRunRef& BuildCookRun, bool bWantToCook, ELauncherProfileDeploymentModes::Type DefaultDeploymentMode)
 	{
 		bool bPakFiles = (ContentScheme == EContentScheme::PakFiles || ContentScheme == EContentScheme::DevelopmentPackage || ContentScheme == EContentScheme::SubmissionPackage);
 		bool bUseZen = (ContentScheme != EContentScheme::LooseFiles && ContentScheme != EContentScheme::CookOnTheFly);
@@ -718,46 +762,57 @@ PRAGMA_ENABLE_DEPRECATION_WARNINGS
 		bool bSubmissionPackage = (ContentScheme == EContentScheme::SubmissionPackage);
 		// FIXME: may need to turn off Zen/Pak in UAT somehow if it's set in project defaults?
 
-		Profile->SetUseZenPakStreaming(bZenPakStreaming);
-		Profile->SetDeployWithUnrealPak(bPakFiles);
-		Profile->SetUseZenStore(bUseZen);
-		Profile->SetCreateReleaseVersion(bSubmissionPackage);
-		Profile->SetGeneratePatch(false);
-		Profile->SetCreateDLC(false);
+		BuildCookRun->SetUseZenPakStreaming(bZenPakStreaming);
+		BuildCookRun->SetDeployWithUnrealPak(bPakFiles);
+		BuildCookRun->SetUseZenStore(bUseZen);
+		BuildCookRun->SetCreateReleaseVersion(bSubmissionPackage);
+		BuildCookRun->SetGeneratePatch(false);
+		BuildCookRun->SetCreateDLC(false);
 		if (!bPakFiles)
 		{
-			Profile->SetGenerateChunks(false);
-			Profile->SetUseIoStore(false);
+			BuildCookRun->SetGenerateChunks(false);
+			BuildCookRun->SetUseIoStore(false);
 		}
 
 		ELauncherProfileDeploymentModes::Type DeploymentMode = DefaultDeploymentMode;
 		if (bCOTF)
 		{
-			Profile->SetCookMode(ELauncherProfileCookModes::OnTheFly);
+			BuildCookRun->SetCookMode(ELauncherProfileCookModes::OnTheFly);
 			DeploymentMode = ELauncherProfileDeploymentModes::FileServer;
 		}
 		else if (bZenPakStreaming || !bWantToCook)
 		{
-			Profile->SetCookMode(ELauncherProfileCookModes::DoNotCook);
+			BuildCookRun->SetCookMode(ELauncherProfileCookModes::DoNotCook);
 		}
 		else
 		{
-			Profile->SetCookMode(ELauncherProfileCookModes::ByTheBook);
+			BuildCookRun->SetCookMode(ELauncherProfileCookModes::ByTheBook);
 		}
 
 		if (bPackage)
 		{
-			Profile->SetDeploymentMode(ELauncherProfileDeploymentModes::DoNotDeploy); // @todo: some platforms support package deployment. should update verification to check new function in ITargetPlatformControls
-			Profile->SetPackagingMode(ELauncherProfilePackagingModes::Locally);
+			BuildCookRun->SetDeploymentMode(ELauncherProfileDeploymentModes::DoNotDeploy); // @todo: some platforms support package deployment. should update verification to check new function in ITargetPlatformControls
+			BuildCookRun->SetPackagingMode(ELauncherProfilePackagingModes::Locally);
 		}
 		else
 		{
-			Profile->SetDeploymentMode(DeploymentMode);
-			Profile->SetPackagingMode(ELauncherProfilePackagingModes::DoNotPackage);
+			BuildCookRun->SetDeploymentMode(DeploymentMode);
+			BuildCookRun->SetPackagingMode(ELauncherProfilePackagingModes::DoNotPackage);
 		}
 	}
 
 
+	ILauncherProfileBuildCookRunRef FModel::CreateBuildCookRun( const ILauncherProfileRef& Profile )
+	{
+		ILauncherProfileBuildCookRunRef BuildCookRun = Profile->FindOrAddBuildCookRunCommand(nullptr);
+
+		SetDefaults(Profile, BuildCookRun);
+
+		return BuildCookRun;
+	}
+
+
+
 	void FModel::ReadProjectSettingsFromConfig( FConfigCacheIni& InConfig, const FString& InProjectPath, FProjectSettings& OutResult )
 	{
 		// read build settings
@@ -902,6 +957,16 @@ PRAGMA_ENABLE_DEPRECATION_WARNINGS
 	}
 
 
+	const ILauncherProfileBuildCookRunRef FModel::GetBasicDefaultBuildCookRun() const
+	{
+		return DefaultBasicLaunchProfile->GetFirstBuildCookRun().ToSharedRef();
+	}
+
+	const ILauncherProfileBuildCookRunRef FModel::GetCustomDefaultBuildCookRun() const
+	{
+		return DefaultCustomLaunchProfile->GetFirstBuildCookRun().ToSharedRef();
+	}
+
 
 	TSharedPtr<FLaunchLogMessage> FModel::AddLogMessage( const FString& InMessage, ELogVerbosity::Type InVerbosity )
 	{
@@ -947,6 +1012,20 @@ PRAGMA_ENABLE_DEPRECATION_WARNINGS
 #endif
 	}
 
+	bool FModel::CanUseSimplifiedLayout(const ILauncherProfileRef& Profile)
+	{
+		if (IsBasicLaunchProfile(Profile))
+		{
+			return true;
+		}
+
+		if (!ProjectLauncher::bUseClassicView)
+		{
+			return false;
+		}
+
+		return (Profile->GetBuildCookRunCommands().Num() == 1 && Profile->GetAutomatedTests().Num() == 0 && Profile->GetUATCommands().Num() == 1);
+	}
 
 	TArray<FString> FModel::GetAndCacheMapPaths(const FString& InOptionalProjectPath, bool bIncludeNonContentDirMaps) // @todo: map list parsing should ideally be asyncronous, showing a spinner in the map selector controls until its finished etc.
 	{
@@ -1172,64 +1251,92 @@ PRAGMA_ENABLE_DEPRECATION_WARNINGS
 
 
 
-	FText GetProfileLaunchErrorMessage(ILauncherProfilePtr Profile, bool bWithAnnotations)
+
+	static void GetUATCommandErrorMessage(FTextBuilder& MsgTextBuilder, ILauncherProfilePtr Profile, const ILauncherProfileUATCommandPtr& UATCommand, bool bWithAnnotations, bool bWithUATCommand, bool& bOutHasErrors, bool& bOutHasWarnings)
 	{
-		if (!Profile.IsValid())
+		TArray<FString> CustomErrors = Profile->GetAllCustomErrors(UATCommand);
+		TArray<FString> CustomWarnings = Profile->GetAllCustomWarnings(UATCommand);
+
+		bool bShowUATCommand = (bWithUATCommand && UATCommand.IsValid() && (CustomErrors.Num() > 0 || CustomWarnings.Num() > 0 || Profile->HasValidationError(UATCommand)));
+		if (bShowUATCommand)
 		{
-			return LOCTEXT("LaunchErrNoProfileTip", "There is no profile selected");
+			MsgTextBuilder.AppendLine(FText::FromString(UATCommand->GetDescription()));
+			MsgTextBuilder.Indent();
 		}
 
-		FTextBuilder MsgTextBuilder;
-
-		TArray<FString> CustomErrors = Profile->GetAllCustomErrors();
-		if (CustomErrors.Num() > 0 || Profile->HasValidationError())
+		if (CustomErrors.Num() > 0 || Profile->HasValidationError(UATCommand))
 		{
-			if (bWithAnnotations)
-			{
-				MsgTextBuilder.AppendLine(LOCTEXT("LaunchErrValidation", "There are validation errors with this profile. Please fix them before launching:"));
-				MsgTextBuilder.Indent();
-			}
+			bOutHasErrors = true;
 			for (int i = 0; i < (int)ELauncherProfileValidationErrors::Count; i++)
 			{
 				ELauncherProfileValidationErrors::Type Error = (ELauncherProfileValidationErrors::Type)i;
-				if (Profile->HasValidationError(Error))
+				if (Profile->HasValidationError(Error, UATCommand))
 				{
-					MsgTextBuilder.AppendLine(LexToStringLocalized(Error));
+					MsgTextBuilder.AppendLineFormat(LOCTEXT("ValidationErrFmt", "Error: {0}"), FText::FromString(LexToStringLocalized(Error)));
 				}
 			}
 			for (const FString& CustomError : CustomErrors)
 			{
-				MsgTextBuilder.AppendLine(Profile->GetCustomErrorText(CustomError));
+				MsgTextBuilder.AppendLineFormat(LOCTEXT("CustomErrFmt", "Error: {0}"), Profile->GetCustomErrorText(CustomError, UATCommand));
 			}
-			if (bWithAnnotations)
+		}
+
+		if (CustomWarnings.Num() > 0)
+		{
+			bOutHasWarnings = true;
+			for (const FString& CustomWarning : CustomWarnings)
 			{
-				MsgTextBuilder.Unindent();
+				MsgTextBuilder.AppendLineFormat(LOCTEXT("CustomWarnFmt", "Warning: {0}"), Profile->GetCustomWarningText(CustomWarning, UATCommand));
 			}
 		}
-		else if (bWithAnnotations)
+
+		if (bShowUATCommand)
 		{
-			MsgTextBuilder.AppendLine(LOCTEXT("LaunchProfileTip", "Launch this profile now"));
+			MsgTextBuilder.Unindent();
+			MsgTextBuilder.AppendLine();
 		}
+	}
 
-		TArray<FString> CustomWarnings = Profile->GetAllCustomWarnings();
-		if (CustomWarnings.Num() > 0)
+	FText GetProfileLaunchErrorMessage(ILauncherProfilePtr Profile, bool bWithAnnotations)
+	{
+		if (!Profile.IsValid())
+		{
+			return LOCTEXT("LaunchErrNoProfileTip", "There is no profile selected");
+		}
+
+		bool bHasErrors = false;
+		bool bHasWarnings = false;
+		FTextBuilder MsgTextBuilder;
+
+		if (Profile->GetFirstBuildCookRun() == nullptr)
 		{
-			if (bWithAnnotations)
+			// global errors are merged into the first build cook run for legacy PL, so we only need to list them separately if we don't have a build cook run
+			GetUATCommandErrorMessage(MsgTextBuilder, Profile, nullptr, bWithAnnotations, false, bHasErrors, bHasWarnings);
+		}
+
+		bool bWithUATCommand = (Profile->GetUATCommands().Num() > 1);
+		for (const ILauncherProfileUATCommandRef& UATCommand : Profile->GetUATCommands())
+		{
+			GetUATCommandErrorMessage(MsgTextBuilder, Profile, UATCommand.ToSharedPtr(), bWithAnnotations, bWithUATCommand, bHasErrors, bHasWarnings);
+		}
+
+		if (bWithAnnotations)
+		{
+			if (bHasErrors)
 			{
-				MsgTextBuilder.AppendLine();
-				MsgTextBuilder.AppendLine(LOCTEXT("LaunchWarnValidation", "There are validation warnings with this profile but these will not prevent launching:"));
-				MsgTextBuilder.Indent();
+				MsgTextBuilder.AppendLine(LOCTEXT("LaunchErrValidation", "There are validation errors with this profile. Please fix them before launching."));
 			}
-			for (const FString& CustomWarning : CustomWarnings)
+			if (bHasWarnings)
 			{
-				MsgTextBuilder.AppendLine(Profile->GetCustomWarningText(CustomWarning));
+				MsgTextBuilder.AppendLine(LOCTEXT("LaunchWarnValidation", "There are validation warnings with this profile but these will not prevent launching."));
 			}
-			if (bWithAnnotations)
+			if (!bHasErrors)
 			{
-				MsgTextBuilder.Unindent();
+				MsgTextBuilder.AppendLine(LOCTEXT("LaunchProfileTip", "Launch this profile now"));
 			}
 		}
 
+
 		return MsgTextBuilder.ToText();
 	}
 }
diff --git a/Engine/Plugins/Developer/ProjectLauncher/Source/ProjectLauncher/Private/ProfileTree/BasicProfileTreeBuilder.cpp b/Engine/Plugins/Developer/ProjectLauncher/Source/ProjectLauncher/Private/ProfileTree/BasicProfileTreeBuilder.cpp
--- a/Engine/Plugins/Developer/ProjectLauncher/Source/ProjectLauncher/Private/ProfileTree/BasicProfileTreeBuilder.cpp
+++ b/Engine/Plugins/Developer/ProjectLauncher/Source/ProjectLauncher/Private/ProfileTree/BasicProfileTreeBuilder.cpp
@@ -9,6 +9,7 @@ namespace ProjectLauncher
 	FBasicProfileTreeBuilder::FBasicProfileTreeBuilder( const ILauncherProfileRef& InProfile, const TSharedRef<FModel>& InModel )
 		: FGenericProfileTreeBuilder( InProfile, InModel->GetDefaultBasicLaunchProfile(), InModel )
 	{
+		DefaultBuildCookRun = Model->GetBasicDefaultBuildCookRun();
 	}
 
 
@@ -16,13 +17,18 @@ namespace ProjectLauncher
 	{
 		FGenericProfileTreeBuilder::Construct();
 
-		FLaunchProfileTreeNode& GeneralSettingsHeader = TreeData->AddHeading( TEXT("GeneralSettings"), LOCTEXT("GeneralSettingsHeading", "General Settings") );
+		ILauncherProfileBuildCookRunRef UATCommand = Profile->GetFirstBuildCookRun().ToSharedRef();
+		FBuildCookRun& BuildCookRun = Get(UATCommand);
+
+		FLaunchProfileTreeNode& GeneralSettingsHeader = TreeData->AddHeading( TEXT("GeneralSettings"), LOCTEXT("GeneralSettingsHeading", "General Settings"), INT32_MIN );
+		GeneralSettingsHeader.UATCommand = UATCommand;
+
 		AddProjectProperty(GeneralSettingsHeader);
-		AddTargetProperty(GeneralSettingsHeader);
-		AddConfigurationProperty(GeneralSettingsHeader);
-		AddContentSchemeProperty(GeneralSettingsHeader);
-		AddTargetDeviceProperty(GeneralSettingsHeader);
-		AddCommandLineProperty(GeneralSettingsHeader);
+		BuildCookRun.AddTargetProperty(GeneralSettingsHeader);
+		BuildCookRun.AddConfigurationProperty(GeneralSettingsHeader);
+		BuildCookRun.AddContentSchemeProperty(GeneralSettingsHeader);
+		BuildCookRun.AddTargetDeviceProperty(GeneralSettingsHeader);
+		BuildCookRun.AddCommandLineProperty(GeneralSettingsHeader);
 	}
 
 
diff --git a/Engine/Plugins/Developer/ProjectLauncher/Source/ProjectLauncher/Private/ProfileTree/CustomProfileTreeBuilder.cpp b/Engine/Plugins/Developer/ProjectLauncher/Source/ProjectLauncher/Private/ProfileTree/CustomProfileTreeBuilder.cpp
--- a/Engine/Plugins/Developer/ProjectLauncher/Source/ProjectLauncher/Private/ProfileTree/CustomProfileTreeBuilder.cpp
+++ b/Engine/Plugins/Developer/ProjectLauncher/Source/ProjectLauncher/Private/ProfileTree/CustomProfileTreeBuilder.cpp
@@ -1,6 +1,9 @@
 // Copyright Epic Games, Inc. All Rights Reserved.
 
 #include "ProfileTree/CustomProfileTreeBuilder.h"
+#include "Model/ProjectLauncherModel.h"
+#include "Extension/LaunchExtension.h"
+#include "Extension/BuildCookRunLaunchExtension.h"
 
 
 #define LOCTEXT_NAMESPACE "CustomProfileTreeBuilder"
@@ -16,46 +19,74 @@ namespace ProjectLauncher
 	{
 		FGenericProfileTreeBuilder::Construct();
 
-		FLaunchProfileTreeNode& GeneralSettingsHeader = TreeData->AddHeading( TEXT("GeneralSettings"), LOCTEXT("GeneralSettingsHeading", "General Settings") );
+		// when there's only one profile and no active BuildCookRun extensions, group the target, platform etc. under the general settings header
+		// @todo: explore whether this can be handled via SCustomLaunchCustomProfileEditor checking bClassicView instead
+		bool bClassicView = Model->CanUseSimplifiedLayout(Profile);
+		auto AddHeading = [this, bClassicView]( FLaunchProfileTreeNode& GeneralBuildCookRunHeader, const TCHAR* Name, const FText& Heading ) -> FLaunchProfileTreeNode&
+		{
+			if (bClassicView)
+			{
+				FLaunchProfileTreeNode& Header = TreeData->AddHeading( Name, Heading );
+				Header.UATCommand = Profile->GetFirstBuildCookRun();
+				return Header;
+			}
+			else
+			{
+				return GeneralBuildCookRunHeader.AddSubHeading( Name, Heading );
+			}
+		};
+
+		FLaunchProfileTreeNode& GeneralSettingsHeader = TreeData->AddHeading( TEXT("GeneralSettings"), LOCTEXT("GeneralSettingsHeading", "General Settings"), INT32_MIN );
+		if (bClassicView)
+		{
+			GeneralSettingsHeader.UATCommand = Profile->GetFirstBuildCookRun();
+		}
+
 		AddProjectProperty(GeneralSettingsHeader);
-		AddTargetProperty(GeneralSettingsHeader);
-		AddPlatformProperty(GeneralSettingsHeader);
-		AddConfigurationProperty(GeneralSettingsHeader);
-		AddContentSchemeProperty(GeneralSettingsHeader);
-
-		FLaunchProfileTreeNode& ContentSchemeHeader = TreeData->AddHeading( TEXT("ContentScheme"), LOCTEXT("ContentSchemeHeading", "Content Scheme") );
-		AddCompressPakFilesProperty(ContentSchemeHeader);
-		AddUseIoStoreProperty(ContentSchemeHeader);
-		AddGenerateChunksProperty(ContentSchemeHeader);
-		AddZenSnapshotProperty(ContentSchemeHeader);
-		AddImportZenSnapshotProperty(ContentSchemeHeader);
-		AddZenPakStreamingPathProperty(ContentSchemeHeader);
-
-		FLaunchProfileTreeNode& SubmissionHeader = TreeData->AddHeading( TEXT("Release"), LOCTEXT("ReleaseHeading", "Release / DLC / Patching") );
-		AddSubmissionPackageProperty(SubmissionHeader);
-
-		FLaunchProfileTreeNode& CookingHeader = TreeData->AddHeading( TEXT("Cooking"), LOCTEXT("CookingHeading", "Cook") );
-		AddCookProperty(CookingHeader);
-		AddIncrementalCookProperty(CookingHeader);
-
-		FLaunchProfileTreeNode& AdvCookingHeader = TreeData->AddHeading( TEXT("AdvancedCook"), LOCTEXT("AdvancedCookingHeading", "Advanced Cooking Options") ); // @fixme: hack - just so Advanced Cooking Extension panel appears directly below normal cook. @todo: Need 'InsertAfter' for headings
-
-
-		FLaunchProfileTreeNode& BuildHeader = TreeData->AddHeading( TEXT("Build"), LOCTEXT("BuildHeading", "Build") );
-		AddBuildProperty(BuildHeader);
-		AddForceBuildProperty(BuildHeader);
-		AddArchitectureProperty(BuildHeader);
-
-		FLaunchProfileTreeNode& DirectoryHeader = TreeData->AddHeading( TEXT("Directory"), LOCTEXT("DirectoryHeading", "Directory") );
-		AddArchiveBuildProperty(DirectoryHeader);
-		AddArchiveBuildDirectoryProperty(DirectoryHeader);
-
-		FLaunchProfileTreeNode& DeployAndRunHeader = TreeData->AddHeading( TEXT("DeployAndRun"), LOCTEXT("DeployAndRunHeading", "Deploy And Run") );
-		AddDeployProperty(DeployAndRunHeader);
-		AddTargetDeviceProperty(DeployAndRunHeader);
-		AddRunProperty(DeployAndRunHeader);
-		AddInitialMapProperty(DeployAndRunHeader);
-		AddCommandLineProperty(DeployAndRunHeader);
+		for (ILauncherProfileBuildCookRunRef UATCommand : Profile->GetBuildCookRunCommands())
+		{
+			FBuildCookRun& BuildCookRun = Get(UATCommand);
+
+			const bool bIsPrimaryHeading = true;
+			FLaunchProfileTreeNode& GeneralBuildCookRunHeader = bClassicView ? GeneralSettingsHeader : TreeData->FindOrAddHeading(UATCommand, bIsPrimaryHeading);
+
+			BuildCookRun.AddTargetProperty(GeneralBuildCookRunHeader);
+			BuildCookRun.AddPlatformProperty(GeneralBuildCookRunHeader);
+			BuildCookRun.AddConfigurationProperty(GeneralBuildCookRunHeader);
+			BuildCookRun.AddContentSchemeProperty(GeneralBuildCookRunHeader);
+
+			FLaunchProfileTreeNode& ContentSchemeHeader = AddHeading( GeneralBuildCookRunHeader, TEXT("ContentScheme"), LOCTEXT("ContentSchemeHeading", "Content Scheme") );
+			BuildCookRun.AddCompressPakFilesProperty(ContentSchemeHeader);
+			BuildCookRun.AddUseIoStoreProperty(ContentSchemeHeader);
+			BuildCookRun.AddGenerateChunksProperty(ContentSchemeHeader);
+			BuildCookRun.AddZenSnapshotProperty(ContentSchemeHeader);
+			BuildCookRun.AddImportZenSnapshotProperty(ContentSchemeHeader);
+			BuildCookRun.AddZenPakStreamingPathProperty(ContentSchemeHeader);
+
+			FLaunchProfileTreeNode& SubmissionHeader = AddHeading( GeneralBuildCookRunHeader,  TEXT("Release"), LOCTEXT("ReleaseHeading", "Release / DLC / Patching") );
+			BuildCookRun.AddSubmissionPackageProperty(SubmissionHeader);
+
+			FLaunchProfileTreeNode& CookingHeader = AddHeading( GeneralBuildCookRunHeader,  TEXT("Cooking"), LOCTEXT("CookingHeading", "Cook") );
+			BuildCookRun.AddCookProperty(CookingHeader);
+			BuildCookRun.AddIncrementalCookProperty(CookingHeader);
+
+			FLaunchProfileTreeNode& AdvCookingHeader = AddHeading( GeneralBuildCookRunHeader,  TEXT("AdvancedCook"), LOCTEXT("AdvancedCookingHeading", "Advanced Cooking Options") ); // @fixme: hack - just so Advanced Cooking Extension panel appears directly below normal cook. @todo: Need 'InsertAfter' for headings
+			FLaunchProfileTreeNode& BuildHeader = AddHeading( GeneralBuildCookRunHeader,  TEXT("Build"), LOCTEXT("BuildHeading", "Build") );
+			BuildCookRun.AddBuildProperty(BuildHeader);
+			BuildCookRun.AddForceBuildProperty(BuildHeader);
+			BuildCookRun.AddArchitectureProperty(BuildHeader);
+
+			FLaunchProfileTreeNode& DirectoryHeader = AddHeading( GeneralBuildCookRunHeader,  TEXT("Directory"), LOCTEXT("DirectoryHeading", "Directory") );
+			BuildCookRun.AddArchiveBuildProperty(DirectoryHeader);
+			BuildCookRun.AddArchiveBuildDirectoryProperty(DirectoryHeader);
+
+			FLaunchProfileTreeNode& DeployAndRunHeader = AddHeading( GeneralBuildCookRunHeader,  TEXT("DeployAndRun"), LOCTEXT("DeployAndRunHeading", "Deploy And Run") );
+			BuildCookRun.AddDeployProperty(DeployAndRunHeader);
+			BuildCookRun.AddTargetDeviceProperty(DeployAndRunHeader);
+			BuildCookRun.AddRunProperty(DeployAndRunHeader);
+			BuildCookRun.AddInitialMapProperty(DeployAndRunHeader);
+			BuildCookRun.AddCommandLineProperty(DeployAndRunHeader);
+		}
 	}
 
 
diff --git a/Engine/Plugins/Developer/ProjectLauncher/Source/ProjectLauncher/Private/ProfileTree/GenericProfileTreeBuilder.cpp b/Engine/Plugins/Developer/ProjectLauncher/Source/ProjectLauncher/Private/ProfileTree/GenericProfileTreeBuilder.cpp
--- a/Engine/Plugins/Developer/ProjectLauncher/Source/ProjectLauncher/Private/ProfileTree/GenericProfileTreeBuilder.cpp
+++ b/Engine/Plugins/Developer/ProjectLauncher/Source/ProjectLauncher/Private/ProfileTree/GenericProfileTreeBuilder.cpp
@@ -35,6 +35,30 @@ namespace ProjectLauncher
 		, Model(InModel)
 	{
 		ProfileType = Model->GetProfileType(Profile);
+		DefaultBuildCookRun = Model->GetCustomDefaultBuildCookRun();
+
+		for (const ILauncherProfileBuildCookRunRef& BuildCookRun : InProfile->GetBuildCookRunCommands())
+		{
+			BuildCookRunData.Add(BuildCookRun, MakeShared<FBuildCookRun>(*this, BuildCookRun, Model));
+		}
+
+		Profile->OnUATCommandAdded().AddRaw( this, &FGenericProfileTreeBuilder::OnUATCommandAdded );
+		Profile->OnUATCommandRemoved().AddRaw( this, &FGenericProfileTreeBuilder::OnUATCommandRemoved );
+	}
+
+	FGenericProfileTreeBuilder::~FGenericProfileTreeBuilder()
+	{
+		Profile->OnUATCommandAdded().RemoveAll(this);
+		Profile->OnUATCommandRemoved().RemoveAll(this);
+	}
+
+
+
+	FGenericProfileTreeBuilder::FBuildCookRun::FBuildCookRun( FGenericProfileTreeBuilder& InOwner, ILauncherProfileBuildCookRunRef InBuildCookRun, const TSharedRef<FModel>& InModel )
+		: Owner(InOwner)
+		, BuildCookRun(InBuildCookRun)
+		, Model(InModel)
+	{
 
 		ForPak = [this]()
 		{
@@ -91,35 +115,48 @@ namespace ProjectLauncher
 			return FString();
 		};
 
-		Profile->OnCustomValidation().AddRaw( this, &FGenericProfileTreeBuilder::OnValidateProfile );
+		Owner.Profile->OnCustomUATCommandValidation().AddRaw( this, &FBuildCookRun::OnValidateProfile );
 	}
 	
-	FGenericProfileTreeBuilder::~FGenericProfileTreeBuilder()
+	FGenericProfileTreeBuilder::FBuildCookRun::~FBuildCookRun()
 	{
-		Profile->OnCustomValidation().RemoveAll( this );
+		Owner.Profile->OnCustomUATCommandValidation().RemoveAll( this );
 	}
 
+
 	void FGenericProfileTreeBuilder::Construct()
 	{
-		const TArray<FString>& DeviceIDs = Profile->GetDeployedDeviceGroup()->GetDeviceIDs();
+		for (const TPair<ILauncherProfileBuildCookRunRef,TSharedRef<FBuildCookRun>>& Pair : BuildCookRunData)
+		{
+			Pair.Value->Construct();
+		}
+	}
+
+	void FGenericProfileTreeBuilder::FBuildCookRun::Construct()
+	{
+		const TArray<FString>& DeviceIDs = BuildCookRun->GetDeployedDeviceGroup()->GetDeviceIDs();
 
 		bShouldCook = GetCook();
-		bHasAdvancedPlatformTargets = (Profile->GetCookedPlatforms().Num() > 1 || Profile->GetBuildTargets().Num() > 1);
-		ContentScheme = Model->DetermineProfileContentScheme(Profile);
+		bHasAdvancedPlatformTargets = (BuildCookRun->GetCookedPlatforms().Num() > 1 || BuildCookRun->GetBuildTargets().Num() > 1);
+		ContentScheme = Model->DetermineProfileContentScheme(BuildCookRun);
 
 		// ensure the default deploy platform matches the cook platform because Legacy Project Launcher can create a mismatch
-		if (Profile->GetDefaultDeployPlatform() != NAME_None && !GetSelectedPlatforms().Contains(Profile->GetDefaultDeployPlatform().ToString()) )
+		if (BuildCookRun->GetDefaultDeployPlatform() != NAME_None && !GetSelectedPlatforms().Contains(BuildCookRun->GetDefaultDeployPlatform().ToString()) )
 		{
 			FString DeployPlatform = (GetSelectedPlatforms().Num() > 0) ? GetSelectedPlatforms()[0] : FPlatformProperties::IniPlatformName();
-			Profile->SetDefaultDeployPlatform(FName(*DeployPlatform));
+			BuildCookRun->SetDefaultDeployPlatform(FName(*DeployPlatform));
 		}
-	
-		CacheArchitectures();
-		CacheBuildTargets();
-	}
 
+		if (InitalMapCombo.IsValid())
+		{
+			InitalMapCombo->SetSelectedItem( GetInitialMap(), ESelectInfo::Direct );
+			InitalMapCombo->RefreshOptions();
+		}
 
+		CacheArchitectures();
+		CacheBuildTargets();
 
+	}
 
 
 
@@ -162,23 +199,23 @@ namespace ProjectLauncher
 		}
 	}
 
-	void FGenericProfileTreeBuilder::AddTargetProperty( FLaunchProfileTreeNode& HeadingNode )
+	void FGenericProfileTreeBuilder::FBuildCookRun::AddTargetProperty( FLaunchProfileTreeNode& HeadingNode )
 	{
-		if (ProfileType == EProfileType::Basic)
+		if (Owner.ProfileType == EProfileType::Basic)
 		{
 			// this is also adding a new widget to the property tree, as above.
 			// in this example the new struct parameter defines several callbacks that handle the 'reset to default' functionality. there are also options for disabling & hiding.
 			// the code is implemented with this syntax to avoid aid readability without using slate's TAttribute style functionality which seemed like an overkill for our simpler needs.
 			HeadingNode.AddWidget( LOCTEXT("TargetLabel", "Target"), 
 				{
-					.IsDefault = [this]()		{ return !Profile->HasBuildTargetSpecified() || (Profile->GetBuildTargets().Num() == 0); },
+					.IsDefault = [this]()		{ return !BuildCookRun->HasBuildTargetSpecified() || (BuildCookRun->GetBuildTargets().Num() == 0); },
 					.SetToDefault = [this]()	{ SetBuildTargets(TArray<FString>()); },
 					.Validation = FValidation( {ELauncherProfileValidationErrors::BuildTargetCookVariantMismatch, ELauncherProfileValidationErrors::BuildTargetIsRequired, ELauncherProfileValidationErrors::CookOnTheFlyDoesntSupportServer} ),
 				},
 				SAssignNew(BuildTargetCombo, SCustomLaunchBuildTargetCombo, Model)
-				.OnSelectionChanged(this, &FGenericProfileTreeBuilder::SetBuildTargets)
-				.SelectedBuildTargets(this, &FGenericProfileTreeBuilder::GetBuildTargets)
-				.SelectedProject(this, &FGenericProfileTreeBuilder::GetProjectPath)
+				.OnSelectionChanged(this, &FGenericProfileTreeBuilder::FBuildCookRun::SetBuildTargets)
+				.SelectedBuildTargets(this, &FGenericProfileTreeBuilder::FBuildCookRun::GetBuildTargets)
+				.SelectedProject(Owner.AsShared(), &FGenericProfileTreeBuilder::GetProjectPath)
 				.Font(FCoreStyle::Get().GetFontStyle("SmallFont"))
 			);
 		}
@@ -186,39 +223,39 @@ namespace ProjectLauncher
 		{
 			HeadingNode.AddWidget( LOCTEXT("TargetLabel", "Target"), 
 				{
-					.IsDefault = [this]()		{ return !Profile->HasBuildTargetSpecified() || (Profile->GetBuildTargets().Num() == 0); },
+					.IsDefault = [this]()		{ return !BuildCookRun->HasBuildTargetSpecified() || (BuildCookRun->GetBuildTargets().Num() == 0); },
 					.SetToDefault = [this]()	{ SetBuildTargets(TArray<FString>()); },
-					.IsEnabled = [this]()		{ return Profile->HasProjectSpecified(); },
+					.IsEnabled = [this]()		{ return Owner.Profile->HasProjectSpecified(); },
 					.Validation = FValidation( {ELauncherProfileValidationErrors::BuildTargetCookVariantMismatch, ELauncherProfileValidationErrors::BuildTargetIsRequired, ELauncherProfileValidationErrors::CookOnTheFlyDoesntSupportServer}, {TEXT("Validation_GameAndClientTargetsDisallowed")}),
 				},
 				SAssignNew(BuildTargetCombo, SCustomLaunchBuildTargetCombo, Model)
-				.OnSelectionChanged(this, &FGenericProfileTreeBuilder::SetBuildTargets)
-				.SelectedBuildTargets(this, &FGenericProfileTreeBuilder::GetBuildTargets)
-				.SelectedProject(this, &FGenericProfileTreeBuilder::GetProjectPath)
-				//.SetAdvancedPlatformsOption(this, &FGenericProfileTreeBuilder::SetHasAdvancedPlatformTargets) // @todo: prevent this from being toggled until the final design on client/server build configuration. only exising multi-target/multi-platform profiles will have multiselectable platforms/target controls
-				.GetAdvancedPlatformsOption(this, &FGenericProfileTreeBuilder::HasAdvancedPlatformTargets)
+				.OnSelectionChanged(this, &FGenericProfileTreeBuilder::FBuildCookRun::SetBuildTargets)
+				.SelectedBuildTargets(this, &FGenericProfileTreeBuilder::FBuildCookRun::GetBuildTargets)
+				.SelectedProject(Owner.AsShared(), &FGenericProfileTreeBuilder::GetProjectPath)
+				//.SetAdvancedPlatformsOption(this, &FGenericProfileTreeBuilder::FBuildCookRun::SetHasAdvancedPlatformTargets) // @todo: prevent this from being toggled until the final design on client/server build configuration. only exising multi-target/multi-platform profiles will have multiselectable platforms/target controls
+				.GetAdvancedPlatformsOption(this, &FGenericProfileTreeBuilder::FBuildCookRun::HasAdvancedPlatformTargets)
 				.Font(FCoreStyle::Get().GetFontStyle("SmallFont"))
 			);
 		}
 	}
 
-	void FGenericProfileTreeBuilder::AddPlatformProperty( FLaunchProfileTreeNode& HeadingNode )
+	void FGenericProfileTreeBuilder::FBuildCookRun::AddPlatformProperty( FLaunchProfileTreeNode& HeadingNode )
 	{
 		HeadingNode.AddWidget( LOCTEXT("PlatformLabel", "Platform"), 
 			{
 				.Validation = FValidation({ELauncherProfileValidationErrors::NoPlatformSelected, ELauncherProfileValidationErrors::NoPlatformSDKInstalled})
 			},
 			SAssignNew(PlatformCombo, SCustomLaunchPlatformCombo)
-			.SelectedPlatforms(this, &FGenericProfileTreeBuilder::GetSelectedPlatforms)
-			.OnSelectionChanged(this, &FGenericProfileTreeBuilder::SetSelectedPlatforms)
-			.GetAdvancedPlatformsOption(this, &FGenericProfileTreeBuilder::HasAdvancedPlatformTargets)
+			.SelectedPlatforms(this, &FGenericProfileTreeBuilder::FBuildCookRun::GetSelectedPlatforms)
+			.OnSelectionChanged(this, &FGenericProfileTreeBuilder::FBuildCookRun::SetSelectedPlatforms)
+			.GetAdvancedPlatformsOption(this, &FGenericProfileTreeBuilder::FBuildCookRun::HasAdvancedPlatformTargets)
 			.SupportedBuildTargets_Lambda( [this]() { return CachedBuildTargetTypes; } )
 			.Font(FCoreStyle::Get().GetFontStyle("SmallFont"))
 		);
 	}
 
 
-	void FGenericProfileTreeBuilder::AddConfigurationProperty( FLaunchProfileTreeNode& HeadingNode )
+	void FGenericProfileTreeBuilder::FBuildCookRun::AddConfigurationProperty( FLaunchProfileTreeNode& HeadingNode )
 	{
 		TArray<EBuildConfiguration> ValidConfigurations;
 		
@@ -237,13 +274,13 @@ namespace ProjectLauncher
 				.Validation = FValidation({ELauncherProfileValidationErrors::NoBuildConfigurationSelected, ELauncherProfileValidationErrors::ShippingDoesntSupportCommandlineOptionsCantUseCookOnTheFly}),
 			},
 			SNew(SCustomLaunchLexToStringCombo<EBuildConfiguration>)
-			.OnSelectionChanged(this, &FGenericProfileTreeBuilder::SetBuildConfiguration)
-			.SelectedItem(this, &FGenericProfileTreeBuilder::GetBuildConfiguration)
+			.OnSelectionChanged(this, &FGenericProfileTreeBuilder::FBuildCookRun::SetBuildConfiguration)
+			.SelectedItem(this, &FGenericProfileTreeBuilder::FBuildCookRun::GetBuildConfiguration)
 			.Items(ValidConfigurations)
 		);
 	}
 
-	void FGenericProfileTreeBuilder::AddContentSchemeProperty( FLaunchProfileTreeNode& HeadingNode )
+	void FGenericProfileTreeBuilder::FBuildCookRun::AddContentSchemeProperty( FLaunchProfileTreeNode& HeadingNode )
 	{
 		HeadingNode.AddWidget( LOCTEXT("ContentSchemeLabel", "Content Scheme"), 
 			{
@@ -251,93 +288,93 @@ namespace ProjectLauncher
 											{TEXT("Validation_ZenNeedsIoStore"), TEXT("Validation_ZenNeedsRemotePermission")} ),
 			},
 			SNew(SCustomLaunchContentSchemeCombo)
-			.OnSelectionChanged(this, &FGenericProfileTreeBuilder::SetContentScheme)
+			.OnSelectionChanged(this, &FGenericProfileTreeBuilder::FBuildCookRun::SetContentScheme)
 			.SelectedContentScheme_Lambda( [this]() { return ContentScheme; } )
-			.IsContentSchemeAvailable(this, &FGenericProfileTreeBuilder::IsContentSchemeAvailable)
+			.IsContentSchemeAvailable(this, &FGenericProfileTreeBuilder::FBuildCookRun::IsContentSchemeAvailable)
 			.Font(FCoreStyle::Get().GetFontStyle("SmallFont"))
 		);
 	}
 
-	void FGenericProfileTreeBuilder::AddCompressPakFilesProperty( FLaunchProfileTreeNode& HeadingNode )
+	void FGenericProfileTreeBuilder::FBuildCookRun::AddCompressPakFilesProperty( FLaunchProfileTreeNode& HeadingNode )
 	{
 		// in this example we are adding a single boolean instead of a custom widget. the struct parameter defines how the value is accessed
 		// the ForPak is the lambda functions created in the constructor and is again aimed at readability
 		HeadingNode.AddBoolean( LOCTEXT("CompressPakFilesLabel", "Compress Pak Files"),
 			{
-				.GetValue = [this]()			{ return Profile->IsCompressed(); },
-				.SetValue = [this](bool bValue)	{ Profile->SetCompressed(bValue); },
-				.GetDefaultValue = [this]()		{ return DefaultProfile->IsCompressed(); },
+				.GetValue = [this]()			{ return BuildCookRun->IsCompressed(); },
+				.SetValue = [this](bool bValue)	{ BuildCookRun->SetCompressed(bValue); },
+				.GetDefaultValue = [this]()		{ return Owner.DefaultBuildCookRun->IsCompressed(); },
 				.IsVisible = ForPak,
 			}
 		);
 	}
 
-	void FGenericProfileTreeBuilder::AddUseIoStoreProperty( FLaunchProfileTreeNode& HeadingNode )
+	void FGenericProfileTreeBuilder::FBuildCookRun::AddUseIoStoreProperty( FLaunchProfileTreeNode& HeadingNode )
 	{
 		auto IoStorePropertyVisible = [this]()
 		{
-			return ForPak()  && !Model->GetProjectSettings(Profile).bUseIoStore;
+			return ForPak()  && !Model->GetProjectSettings(Owner.Profile).bUseIoStore;
 		};
 
 		HeadingNode.AddBoolean( LOCTEXT("UseIoStoreLabel", "Use Io Store"),
 			{
-				.GetValue = [this]()			{ return Profile->IsUsingIoStore(); },
-				.SetValue = [this](bool bValue)	{ Profile->SetUseIoStore(bValue); },
-				.GetDefaultValue = [this]()		{ return DefaultProfile->IsUsingIoStore(); },
+				.GetValue = [this]()			{ return BuildCookRun->IsUsingIoStore(); },
+				.SetValue = [this](bool bValue)	{ BuildCookRun->SetUseIoStore(bValue); },
+				.GetDefaultValue = [this]()		{ return Owner.DefaultBuildCookRun->IsUsingIoStore(); },
 				.IsVisible = IoStorePropertyVisible,
 			}
 		);
 	}
 
-	void FGenericProfileTreeBuilder::AddGenerateChunksProperty( FLaunchProfileTreeNode& HeadingNode )
+	void FGenericProfileTreeBuilder::FBuildCookRun::AddGenerateChunksProperty( FLaunchProfileTreeNode& HeadingNode )
 	{
 		HeadingNode.AddBoolean( LOCTEXT("GenerateChunksLabel", "Generate Chunks"),
 			{
-				.GetValue = [this]()			{ return Profile->IsGeneratingChunks(); },
-				.SetValue = [this](bool bValue)	{ Profile->SetGenerateChunks(bValue); },
-				.GetDefaultValue = [this]()		{ return DefaultProfile->IsGeneratingChunks(); },
+				.GetValue = [this]()			{ return BuildCookRun->IsGeneratingChunks(); },
+				.SetValue = [this](bool bValue)	{ BuildCookRun->SetGenerateChunks(bValue); },
+				.GetDefaultValue = [this]()		{ return Owner.DefaultBuildCookRun->IsGeneratingChunks(); },
 				.IsVisible = ForPak,
 				.Validation = FValidation({ELauncherProfileValidationErrors::GeneratingChunksRequiresCookByTheBook, ELauncherProfileValidationErrors::GeneratingChunksRequiresUnrealPak}),
 			}
 		);
 	}
 
-	void FGenericProfileTreeBuilder::AddZenSnapshotProperty(FLaunchProfileTreeNode& HeadingNode)
+	void FGenericProfileTreeBuilder::FBuildCookRun::AddZenSnapshotProperty(FLaunchProfileTreeNode& HeadingNode) // @todo: could revisit how this is displayed as it's global? could be inline next to AddImportZenSnapshotProperty tickbox?
 	{
 		HeadingNode.AddWidget(LOCTEXT("ZenSnapshotLabel", "Closest Zen Snapshot"),
 			{
 				.IsVisible = [this]() { return ForContent(); }
 			},
 			SNew(STextBlock)
-			.Text(this, &FGenericProfileTreeBuilder::GetZenSnapshotText)
+			.Text(this, &FGenericProfileTreeBuilder::FBuildCookRun::GetZenSnapshotText)
 			.Font(FCoreStyle::Get().GetFontStyle("SmallFont")));
 	}
 
-	FText FGenericProfileTreeBuilder::GetZenSnapshotText() const
+	FText FGenericProfileTreeBuilder::FBuildCookRun::GetZenSnapshotText() const
 	{
-		int32 Build = Profile->GetZenSnapshot();
+		int32 Build = Owner.Profile->GetZenSnapshot();
 		return FText::FromString(FString::FromInt(Build));
 	}
 
-	void FGenericProfileTreeBuilder::AddImportZenSnapshotProperty( FLaunchProfileTreeNode& HeadingNode )
+	void FGenericProfileTreeBuilder::FBuildCookRun::AddImportZenSnapshotProperty( FLaunchProfileTreeNode& HeadingNode )
 	{
 		HeadingNode.AddBoolean( LOCTEXT("ImportZenSnapshotLabel", "Import Best Match Zen Snapshot"),
 			{
-				.GetValue = [this]()			{ return Profile->IsImportingZenSnapshot(); },
-				.SetValue = [this](bool bValue)	{ Profile->SetImportingZenSnapshot(bValue); },
-				.GetDefaultValue = [this]()		{ return DefaultProfile->IsImportingZenSnapshot(); },
+				.GetValue = [this]()			{ return BuildCookRun->IsImportingZenSnapshot(); },
+				.SetValue = [this](bool bValue)	{ BuildCookRun->SetImportingZenSnapshot(bValue); },
+				.GetDefaultValue = [this]()		{ return Owner.DefaultBuildCookRun->IsImportingZenSnapshot(); },
 				.IsVisible = ForContent,
 				.Validation = FValidation({ELauncherProfileValidationErrors::CopyToDeviceRequiresCookByTheBook}),
 			}
 		);
 	}
 
-	void FGenericProfileTreeBuilder::AddZenPakStreamingPathProperty( FLaunchProfileTreeNode& HeadingNode )
+	void FGenericProfileTreeBuilder::FBuildCookRun::AddZenPakStreamingPathProperty( FLaunchProfileTreeNode& HeadingNode )
 	{
 		HeadingNode.AddDirectoryString( LOCTEXT("ZenPakStreamingPathLabel", "Zen Pak Streaming Path"),
 			{
-				.GetValue = [this]()				{ return Profile->GetZenPakStreamingPath(); },
-				.SetValue = [this](FString Value)	{ Profile->SetZenPakStreamingPath(Value); },
+				.GetValue = [this]()				{ return BuildCookRun->GetZenPakStreamingPath(); },
+				.SetValue = [this](FString Value)	{ BuildCookRun->SetZenPakStreamingPath(Value); },
 				.GetDefaultValue = EmptyString,
 				.IsVisible = ForZenWS,
 				.Validation = FValidation({ELauncherProfileValidationErrors::ZenPakStreamingRequiresDeployAndLaunch}),
@@ -345,11 +382,11 @@ namespace ProjectLauncher
 		);
 	}
 
-	void FGenericProfileTreeBuilder::AddSubmissionPackageProperty( FLaunchProfileTreeNode& HeadingNode )
+	void FGenericProfileTreeBuilder::FBuildCookRun::AddSubmissionPackageProperty( FLaunchProfileTreeNode& HeadingNode )
 	{
 		auto TryGetTargetPlatform = [this](const ITargetPlatform*& OutPlatform)
 		{
-			const TArray<FString>& CookedPlatforms = Profile->GetCookedPlatforms();
+			const TArray<FString>& CookedPlatforms = BuildCookRun->GetCookedPlatforms();
 			if (CookedPlatforms.Num() == 1)
 			{
 				const TArray<ITargetPlatform*>& TargetPlatforms = GetTargetPlatformManager()->GetTargetPlatforms();
@@ -375,13 +412,13 @@ namespace ProjectLauncher
 
 		auto IsBasedOnReleaseVersionVisible = [this]
 		{
-			return ForSubmissionPackage() && (Profile->IsCreatingDLC() || Profile->IsGeneratingPatch());
+			return ForSubmissionPackage() && (BuildCookRun->IsCreatingDLC() || BuildCookRun->IsGeneratingPatch());
 		};
 
 		auto IsOriginalReleaseVersionVisible = [this, TryGetTargetPlatform]
 		{
 			const ITargetPlatform* Platform;
-			return ForSubmissionPackage() && TryGetTargetPlatform(Platform) && Platform->RequiresOriginalReleaseVersionForPatch() && Profile->IsGeneratingPatch();
+			return ForSubmissionPackage() && TryGetTargetPlatform(Platform) && Platform->RequiresOriginalReleaseVersionForPatch() && BuildCookRun->IsGeneratingPatch();
 		};
 
 
@@ -405,52 +442,52 @@ namespace ProjectLauncher
 		// ... release version ...
 		.AddBoolean( LOCTEXT("CreateReleaseLabel", "Create Release Version"),
 			{
-				.GetValue = [this]()				{ return Profile->IsCreatingReleaseVersion(); },
-				.SetValue = [this](bool bValue)		{ Profile->SetCreateReleaseVersion(bValue); },
-				.GetDefaultValue = [this]()			{ return DefaultProfile->IsCreatingReleaseVersion(); },
+				.GetValue = [this]()				{ return BuildCookRun->IsCreatingReleaseVersion(); },
+				.SetValue = [this](bool bValue)		{ BuildCookRun->SetCreateReleaseVersion(bValue); },
+				.GetDefaultValue = [this]()			{ return Owner.DefaultBuildCookRun->IsCreatingReleaseVersion(); },
 				.IsVisible = ForSubmissionPackage,
 			}
 		)
 		.AddString( LOCTEXT("ReleaseVersionNameLabel", "Release Version Name"),
 			{
-				.GetValue = [this]()				{ return Profile->GetCreateReleaseVersionName(); },
-				.SetValue = [this](FString Value)	{ Profile->SetCreateReleaseVersionName(Value); },
-				.GetDefaultValue = [this]()			{ return DefaultProfile->GetCreateReleaseVersionName(); },
+				.GetValue = [this]()				{ return BuildCookRun->GetCreateReleaseVersionName(); },
+				.SetValue = [this](FString Value)	{ BuildCookRun->SetCreateReleaseVersionName(Value); },
+				.GetDefaultValue = [this]()			{ return Owner.DefaultBuildCookRun->GetCreateReleaseVersionName(); },
 				.IsVisible = ForSubmissionPackage,
-				.IsEnabled = [this]()				{ return Profile->IsCreatingReleaseVersion() || Profile->IsGeneratingPatch(); },
+				.IsEnabled = [this]()				{ return BuildCookRun->IsCreatingReleaseVersion() || BuildCookRun->IsGeneratingPatch(); },
 			}
 		)
 
 		// ... various release/distribution options ....
 		.AddBoolean( LOCTEXT("EncryptIniLabel", "Encrypt ini files"),
 			{
-				.GetValue = [this]()				{ return Profile->IsEncryptingIniFiles(); },
-				.SetValue = [this](bool bValue)		{ Profile->SetEncryptingIniFiles(bValue); },
-				.GetDefaultValue = [this]()			{ return DefaultProfile->IsEncryptingIniFiles(); },
+				.GetValue = [this]()				{ return BuildCookRun->IsEncryptingIniFiles(); },
+				.SetValue = [this](bool bValue)		{ BuildCookRun->SetEncryptingIniFiles(bValue); },
+				.GetDefaultValue = [this]()			{ return Owner.DefaultBuildCookRun->IsEncryptingIniFiles(); },
 				.IsVisible = ForSubmissionPackage,
 			}
 		)
 		.AddBoolean( LOCTEXT("ForDistributionLabel", "This build is for distribution to the public"),
 			{
-				.GetValue = [this]()				{ return Profile->IsForDistribution(); },
-				.SetValue = [this](bool bValue)		{ Profile->SetForDistribution(bValue); },
-				.GetDefaultValue = [this]()			{ return DefaultProfile->IsForDistribution(); },
+				.GetValue = [this]()				{ return BuildCookRun->IsForDistribution(); },
+				.SetValue = [this](bool bValue)		{ BuildCookRun->SetForDistribution(bValue); },
+				.GetDefaultValue = [this]()			{ return Owner.DefaultBuildCookRun->IsForDistribution(); },
 				.IsVisible = ForSubmissionPackage,
 			}
 		)
 		.AddBoolean( LOCTEXT("IncludePrereqLabel", "Include an installer for prerequisites"),
 			{
-				.GetValue = [this]()				{ return Profile->IsIncludingPrerequisites(); },
-				.SetValue = [this](bool bValue)		{ Profile->SetIncludePrerequisites(bValue); },
-				.GetDefaultValue = [this]()			{ return DefaultProfile->IsIncludingPrerequisites(); },
+				.GetValue = [this]()				{ return BuildCookRun->IsIncludingPrerequisites(); },
+				.SetValue = [this](bool bValue)		{ BuildCookRun->SetIncludePrerequisites(bValue); },
+				.GetDefaultValue = [this]()			{ return Owner.DefaultBuildCookRun->IsIncludingPrerequisites(); },
 				.IsVisible = ForSubmissionPackage,
 			}
 		)
 		.AddBoolean( LOCTEXT("BinaryConfigLabel", "Make a binary config file for faster runtime settings"),
 			{
-				.GetValue = [this]()				{ return Profile->MakeBinaryConfig(); },
-				.SetValue = [this](bool bValue)		{ Profile->SetMakeBinaryConfig(bValue); },
-				.GetDefaultValue = [this]()			{ return DefaultProfile->MakeBinaryConfig(); },
+				.GetValue = [this]()				{ return BuildCookRun->MakeBinaryConfig(); },
+				.SetValue = [this](bool bValue)		{ BuildCookRun->SetMakeBinaryConfig(bValue); },
+				.GetDefaultValue = [this]()			{ return Owner.DefaultBuildCookRun->MakeBinaryConfig(); },
 				.IsVisible = ForSubmissionPackage,
 			}
 		)
@@ -458,63 +495,63 @@ namespace ProjectLauncher
 		// ... patch/DLC options ...
 		.AddBoolean( LOCTEXT("CreatePatchLabel", "Create Patch"),
 			{
-				.GetValue = [this]()				{ return Profile->IsGeneratingPatch(); },
-				.SetValue = [this](bool bValue)		{ Profile->SetGeneratePatch(bValue); },
-				.GetDefaultValue = [this]()			{ return DefaultProfile->IsGeneratingPatch(); },
+				.GetValue = [this]()				{ return BuildCookRun->IsGeneratingPatch(); },
+				.SetValue = [this](bool bValue)		{ BuildCookRun->SetGeneratePatch(bValue); },
+				.GetDefaultValue = [this]()			{ return Owner.DefaultBuildCookRun->IsGeneratingPatch(); },
 				.IsVisible = [this]()				{ return ForSubmissionPackage(); },
-				.IsEnabled = [this]()				{ return !Profile->IsCreatingDLC(); },
+				.IsEnabled = [this]()				{ return !BuildCookRun->IsCreatingDLC(); },
 			}
 		)
 		.AddBoolean( LOCTEXT("CreateDLCLabel", "Create DLC"),
 			{
-				.GetValue = [this]()				{ return Profile->IsCreatingDLC(); },
-				.SetValue = [this](bool bValue)		{ Profile->SetCreateDLC(bValue); },
-				.GetDefaultValue = [this]()			{ return DefaultProfile->IsCreatingDLC(); },
+				.GetValue = [this]()				{ return BuildCookRun->IsCreatingDLC(); },
+				.SetValue = [this](bool bValue)		{ BuildCookRun->SetCreateDLC(bValue); },
+				.GetDefaultValue = [this]()			{ return Owner.DefaultBuildCookRun->IsCreatingDLC(); },
 				.IsVisible = [this]()				{ return ForSubmissionPackage(); },
-				.IsEnabled = [this]()				{ return !Profile->IsGeneratingPatch(); },
+				.IsEnabled = [this]()				{ return !BuildCookRun->IsGeneratingPatch(); },
 			}
 		)
 
 		.AddString( LOCTEXT("BasedOnReleaseVersionNameLabel", "Based On Release Version Name"),
 			{
-				.GetValue = [this]()				{ return Profile->GetBasedOnReleaseVersionName(); },
-				.SetValue = [this](FString Value)	{ Profile->SetBasedOnReleaseVersionName(Value); },
-				.GetDefaultValue = [this]()			{ return DefaultProfile->GetBasedOnReleaseVersionName(); },
+				.GetValue = [this]()				{ return BuildCookRun->GetBasedOnReleaseVersionName(); },
+				.SetValue = [this](FString Value)	{ BuildCookRun->SetBasedOnReleaseVersionName(Value); },
+				.GetDefaultValue = [this]()			{ return Owner.DefaultBuildCookRun->GetBasedOnReleaseVersionName(); },
 				.IsVisible = IsBasedOnReleaseVersionVisible,
 			}
 		)
 		.AddString( LOCTEXT("OriginalReleaseVersionNameLabel", "Original Release Version Name"),
 			{
-				.GetValue = [this]()				{ return Profile->GetOriginalReleaseVersionName(); },
-				.SetValue = [this](FString Value)	{ Profile->SetOriginalReleaseVersionName(Value); },
-				.GetDefaultValue = [this]()			{ return DefaultProfile->GetOriginalReleaseVersionName(); },
+				.GetValue = [this]()				{ return BuildCookRun->GetOriginalReleaseVersionName(); },
+				.SetValue = [this](FString Value)	{ BuildCookRun->SetOriginalReleaseVersionName(Value); },
+				.GetDefaultValue = [this]()			{ return Owner.DefaultBuildCookRun->GetOriginalReleaseVersionName(); },
 				.IsVisible = IsOriginalReleaseVersionVisible,
 			}
 		)
 
 		.AddBoolean( LOCTEXT("AddPatchLevelLabel", "Add Patch Level (not generally needed)"),
 			{
-				.GetValue = [this]()				{ return Profile->ShouldAddPatchLevel(); },
-				.SetValue = [this](bool bValue)		{ Profile->SetAddPatchLevel(bValue); },
-				.GetDefaultValue = [this]()			{ return DefaultProfile->ShouldAddPatchLevel(); },
-				.IsVisible = [this]()				{ return ForSubmissionPackage() && Profile->IsGeneratingPatch(); },
+				.GetValue = [this]()				{ return BuildCookRun->ShouldAddPatchLevel(); },
+				.SetValue = [this](bool bValue)		{ BuildCookRun->SetAddPatchLevel(bValue); },
+				.GetDefaultValue = [this]()			{ return Owner.DefaultBuildCookRun->ShouldAddPatchLevel(); },
+				.IsVisible = [this]()				{ return ForSubmissionPackage() && BuildCookRun->IsGeneratingPatch(); },
 			}
 		)
 
 		.AddString( LOCTEXT("DLCNameLabel", "DLC Name"),
 			{
-				.GetValue = [this]()				{ return Profile->GetDLCName(); },
-				.SetValue = [this](FString Value)	{ Profile->SetDLCName(Value); },
-				.GetDefaultValue = [this]()			{ return DefaultProfile->GetDLCName(); },
-				.IsVisible = [this]()				{ return ForSubmissionPackage() && Profile->IsCreatingDLC(); },
+				.GetValue = [this]()				{ return BuildCookRun->GetDLCName(); },
+				.SetValue = [this](FString Value)	{ BuildCookRun->SetDLCName(Value); },
+				.GetDefaultValue = [this]()			{ return Owner.DefaultBuildCookRun->GetDLCName(); },
+				.IsVisible = [this]()				{ return ForSubmissionPackage() && BuildCookRun->IsCreatingDLC(); },
 			}
 		)
 		.AddBoolean( LOCTEXT("DLCIncludeEngineLabel", "DLC Includes Engine Content"),
 			{
-				.GetValue = [this]()				{ return Profile->IsDLCIncludingEngineContent(); },
-				.SetValue = [this](bool bValue)		{ Profile->SetDLCIncludeEngineContent(bValue); },
-				.GetDefaultValue = [this]()			{ return DefaultProfile->IsDLCIncludingEngineContent(); },
-				.IsVisible = [this]()				{ return ForSubmissionPackage() && Profile->IsCreatingDLC(); },
+				.GetValue = [this]()				{ return BuildCookRun->IsDLCIncludingEngineContent(); },
+				.SetValue = [this](bool bValue)		{ BuildCookRun->SetDLCIncludeEngineContent(bValue); },
+				.GetDefaultValue = [this]()			{ return Owner.DefaultBuildCookRun->IsDLCIncludingEngineContent(); },
+				.IsVisible = [this]()				{ return ForSubmissionPackage() && BuildCookRun->IsCreatingDLC(); },
 			}
 		)
 
@@ -522,13 +559,13 @@ namespace ProjectLauncher
 	}
 
 
-	void FGenericProfileTreeBuilder::AddIncrementalCookProperty( FLaunchProfileTreeNode& HeadingNode )
+	void FGenericProfileTreeBuilder::FBuildCookRun::AddIncrementalCookProperty( FLaunchProfileTreeNode& HeadingNode )
 	{
 		HeadingNode.AddBoolean( LOCTEXT("IncrementalCookLabel", "Incremental Cook"),
 			{
-				.GetValue = [this]()			{ return Profile->IsCookingIncrementally(); },
-				.SetValue = [this](bool bValue)	{ Profile->SetIncrementalCooking(bValue); },
-				.GetDefaultValue = [this]()		{ return DefaultProfile->IsCookingIncrementally(); },
+				.GetValue = [this]()			{ return BuildCookRun->IsCookingIncrementally(); },
+				.SetValue = [this](bool bValue)	{ BuildCookRun->SetIncrementalCooking(bValue); },
+				.GetDefaultValue = [this]()		{ return Owner.DefaultBuildCookRun->IsCookingIncrementally(); },
 				.IsVisible = ForEnabledCooked,
 				.Validation = FValidation({ELauncherProfileValidationErrors::UnversionedAndIncremental}),
 			}
@@ -537,55 +574,43 @@ namespace ProjectLauncher
 
 
 
-	void FGenericProfileTreeBuilder::AddCookProperty( FLaunchProfileTreeNode& HeadingNode )
+	void FGenericProfileTreeBuilder::FBuildCookRun::AddCookProperty( FLaunchProfileTreeNode& HeadingNode )
 	{
 		HeadingNode.AddBoolean( LOCTEXT("CookLabel", "Cook Content"),
 			{
 				.GetValue = [this]()			{ return GetCook(); },
 				.SetValue = [this](bool bValue)	{ SetCook(bValue); },
-				.GetDefaultValue = [this]()		{ return GetCook(DefaultProfile); },
+				.GetDefaultValue = [this]()		{ return GetCook(Owner.DefaultBuildCookRun); },
 				.IsVisible = ForCooked,
 				.Validation = FValidation({ELauncherProfileValidationErrors::CopyToDeviceRequiresCookByTheBook, ELauncherProfileValidationErrors::GeneratingPatchesCanOnlyRunFromByTheBookCookMode, ELauncherProfileValidationErrors::GeneratingChunksRequiresCookByTheBook}),
 			}
 		);
 	}
 
-	void FGenericProfileTreeBuilder::AddBuildProperty( FLaunchProfileTreeNode& HeadingNode )
+	void FGenericProfileTreeBuilder::FBuildCookRun::AddBuildProperty( FLaunchProfileTreeNode& HeadingNode )
 	{
 		HeadingNode.AddBoolean( LOCTEXT("BuildLabel", "Build the game"),
 			{
 				.GetValue = [this]()			{ return GetBuild(); },
 				.SetValue = [this](bool bValue)	{ SetBuild(bValue); },
-				.GetDefaultValue = [this]()		{ return GetBuild(DefaultProfile); },
+				.GetDefaultValue = [this]()		{ return GetBuild(Owner.DefaultBuildCookRun); },
 			}
 		);
 	}
 
-	void FGenericProfileTreeBuilder::AddForceBuildProperty( FLaunchProfileTreeNode& HeadingNode )
+	void FGenericProfileTreeBuilder::FBuildCookRun::AddForceBuildProperty( FLaunchProfileTreeNode& HeadingNode )
 	{
 		HeadingNode.AddBoolean( LOCTEXT("ForceBuildLabel", "Build even if a pre-built target exists"),
 			{
 				.GetValue = [this]()			{ return GetForceBuild(); },
 				.SetValue = [this](bool bValue)	{ SetForceBuild(bValue); },
-				.GetDefaultValue = [this]()		{ return GetForceBuild(DefaultProfile); },
-				.IsEnabled = ForCode,
-			}
-		);
-	}
-
-	void FGenericProfileTreeBuilder::AddBuildUATProperty( FLaunchProfileTreeNode& HeadingNode )
-	{
-		HeadingNode.AddBoolean( LOCTEXT("BuildUATLabel", "Build UAT"),
-			{
-				.GetValue = [this]()			{ return Profile->IsBuildingUAT(); },
-				.SetValue = [this](bool bValue)	{ Profile->SetBuildUAT(bValue); },
-				.GetDefaultValue = [this]()		{ return DefaultProfile->IsBuildingUAT(); },
+				.GetDefaultValue = [this]()		{ return GetForceBuild(Owner.DefaultBuildCookRun); },
 				.IsEnabled = ForCode,
 			}
 		);
 	}
 
-	void FGenericProfileTreeBuilder::AddArchitectureProperty( FLaunchProfileTreeNode& HeadingNode )
+	void FGenericProfileTreeBuilder::FBuildCookRun::AddArchitectureProperty( FLaunchProfileTreeNode& HeadingNode )
 	{
 		static const struct
 		{
@@ -611,9 +636,9 @@ namespace ProjectLauncher
 					.IsVisible = [this, TargetType]()       { return CachedBuildTargetTypes.Contains(TargetType) && CachedBuildTargetArchitectures[TargetType].Num() > 0; }
 				},
 				SNew(SCustomLaunchStringCombo)
-				.OnSelectionChanged( this, &FGenericProfileTreeBuilder::SetArchitecture, TargetType )
-				.SelectedItem( this, &FGenericProfileTreeBuilder::GetArchitecture, TargetType)
-				.GetDisplayName( this, &FGenericProfileTreeBuilder::GetArchitectureDisplayName )
+				.OnSelectionChanged( this, &FGenericProfileTreeBuilder::FBuildCookRun::SetArchitecture, TargetType )
+				.SelectedItem( this, &FGenericProfileTreeBuilder::FBuildCookRun::GetArchitecture, TargetType)
+				.GetDisplayName( this, &FGenericProfileTreeBuilder::FBuildCookRun::GetArchitectureDisplayName )
 				.Items_Lambda( [this, TargetType]() { return CachedBuildTargetArchitectures[TargetType];} )
 			);
 
@@ -623,12 +648,12 @@ namespace ProjectLauncher
 
 
 
-	void FGenericProfileTreeBuilder::AddStagingDirectoryProperty( FLaunchProfileTreeNode& HeadingNode )
+	void FGenericProfileTreeBuilder::FBuildCookRun::AddStagingDirectoryProperty( FLaunchProfileTreeNode& HeadingNode )
 	{
 		HeadingNode.AddDirectoryString( LOCTEXT("CustomStagingPathLabel", "Custom Stage Directory"),
 			{
-				.GetValue = [this]()				{ return Profile->GetPackageDirectory(); },
-				.SetValue = [this](FString Value)	{ Profile->SetPackageDirectory(Value); },
+				.GetValue = [this]()				{ return BuildCookRun->GetPackageDirectory(); },
+				.SetValue = [this](FString Value)	{ BuildCookRun->SetPackageDirectory(Value); },
 				.GetDefaultValue = EmptyString,
 				.IsVisible = ForCooked,
 				.Validation = FValidation({ELauncherProfileValidationErrors::NoPackageDirectorySpecified}),
@@ -636,28 +661,28 @@ namespace ProjectLauncher
 		);
 	}
 
-	void FGenericProfileTreeBuilder::AddArchiveBuildProperty( FLaunchProfileTreeNode& HeadingNode )
+	void FGenericProfileTreeBuilder::FBuildCookRun::AddArchiveBuildProperty( FLaunchProfileTreeNode& HeadingNode )
 	{
 		HeadingNode.AddBoolean( LOCTEXT("ArchiveBuildLabel", "Archive Build"),
 			{
-				.GetValue = [this]()			{ return Profile->IsArchiving(); },
-				.SetValue = [this](bool bValue)	{ Profile->SetArchive(bValue); },
-				.GetDefaultValue = [this]()		{ return DefaultProfile->IsArchiving(); },
+				.GetValue = [this]()			{ return BuildCookRun->IsArchiving(); },
+				.SetValue = [this](bool bValue)	{ BuildCookRun->SetArchive(bValue); },
+				.GetDefaultValue = [this]()		{ return Owner.DefaultBuildCookRun->IsArchiving(); },
 				.IsVisible = ForCooked,
 			}
 		);
 	}
 
 
-	void FGenericProfileTreeBuilder::AddArchiveBuildDirectoryProperty( FLaunchProfileTreeNode& HeadingNode )
+	void FGenericProfileTreeBuilder::FBuildCookRun::AddArchiveBuildDirectoryProperty( FLaunchProfileTreeNode& HeadingNode )
 	{
 		HeadingNode.AddDirectoryString( LOCTEXT("ArchivePathLabel", "Archive Directory"),
 			{
-				.GetValue = [this]()				{ return Profile->GetArchiveDirectory(); },
-				.SetValue = [this](FString Value)	{ Profile->SetArchiveDirectory(Value); },
+				.GetValue = [this]()				{ return BuildCookRun->GetArchiveDirectory(); },
+				.SetValue = [this](FString Value)	{ BuildCookRun->SetArchiveDirectory(Value); },
 				.GetDefaultValue = EmptyString,
 				.IsVisible = ForCooked,				
-				.IsEnabled = [this]()				{ return Profile->IsArchiving(); },
+				.IsEnabled = [this]()				{ return BuildCookRun->IsArchiving(); },
 				.Validation = FValidation({ELauncherProfileValidationErrors::NoArchiveDirectorySpecified}),
 			}
 		);
@@ -665,33 +690,33 @@ namespace ProjectLauncher
 
 
 
-	void FGenericProfileTreeBuilder::AddDeployProperty( FLaunchProfileTreeNode& HeadingNode )
+	void FGenericProfileTreeBuilder::FBuildCookRun::AddDeployProperty( FLaunchProfileTreeNode& HeadingNode )
 	{
 		HeadingNode.AddBoolean( LOCTEXT("DeployLabel", "Deploy To Device"),
 			{
 				.GetValue = [this]()			{ return GetDeployToDevice(); },
 				.SetValue = [this](bool bValue)	{ SetDeployToDevice(bValue); },
-				.GetDefaultValue = [this]()		{ return GetDeployToDevice(DefaultProfile); },
+				.GetDefaultValue = [this]()		{ return GetDeployToDevice(Owner.DefaultBuildCookRun); },
 				.IsVisible = [this]()			{ return (ContentScheme != EContentScheme::CookOnTheFly) && (ContentScheme != EContentScheme::SubmissionPackage); },
 				.Validation = FValidation({ELauncherProfileValidationErrors::CopyToDeviceRequiresCookByTheBook, ELauncherProfileValidationErrors::CopyToDeviceRequiresNoPackaging, ELauncherProfileValidationErrors::ZenPakStreamingRequiresDeployAndLaunch, ELauncherProfileValidationErrors::DeployedDeviceGroupRequired }),
 			}
 		);
 	}
 
-	void FGenericProfileTreeBuilder::AddIncrementalDeployProperty( FLaunchProfileTreeNode& HeadingNode )
+	void FGenericProfileTreeBuilder::FBuildCookRun::AddIncrementalDeployProperty( FLaunchProfileTreeNode& HeadingNode )
 	{
 		HeadingNode.AddBoolean( LOCTEXT("DeployModifiedLabel", "Only Deploy Modified Content"),
 			{
-				.GetValue = [this]()			{ return Profile->IsDeployingIncrementally(); },
-				.SetValue = [this](bool bValue)	{ Profile->SetIncrementalDeploying(bValue); },
-				.GetDefaultValue = [this]()		{ return DefaultProfile->IsDeployingIncrementally(); },
+				.GetValue = [this]()			{ return BuildCookRun->IsDeployingIncrementally(); },
+				.SetValue = [this](bool bValue)	{ BuildCookRun->SetIncrementalDeploying(bValue); },
+				.GetDefaultValue = [this]()		{ return Owner.DefaultBuildCookRun->IsDeployingIncrementally(); },
 				.IsVisible = [this]()			{ return (ContentScheme != EContentScheme::CookOnTheFly) && (ContentScheme != EContentScheme::SubmissionPackage); },
 				.IsEnabled = ForDeployment,
 			}
 		);
 	}
 
-	void FGenericProfileTreeBuilder::AddTargetDeviceProperty( FLaunchProfileTreeNode& HeadingNode )
+	void FGenericProfileTreeBuilder::FBuildCookRun::AddTargetDeviceProperty( FLaunchProfileTreeNode& HeadingNode )
 	{
 		HeadingNode.AddWidget( LOCTEXT("TargetDeviceLabel", "Target Device"), 
 			{
@@ -704,47 +729,47 @@ namespace ProjectLauncher
 	}
 
 
-	void FGenericProfileTreeBuilder::AddRunProperty( FLaunchProfileTreeNode& HeadingNode )
+	void FGenericProfileTreeBuilder::FBuildCookRun::AddRunProperty( FLaunchProfileTreeNode& HeadingNode )
 	{
 		HeadingNode.AddBoolean( LOCTEXT("RunLabel", "Run"),
 			{
 				.GetValue = [this]()			{ return GetIsRunning(); },
 				.SetValue = [this](bool bValue)	{ SetIsRunning(bValue); },
-				.GetDefaultValue = [this]()		{ return GetIsRunning(DefaultProfile); },
-				.IsVisible = [this]()           { return (ContentScheme != EContentScheme::SubmissionPackage) && Profile->GetAutomatedTests().Num() == 0; }, // run is implied when there are automated tests
+				.GetDefaultValue = [this]()		{ return GetIsRunning(Owner.DefaultBuildCookRun); },
+				.IsVisible = [this]()           { return (ContentScheme != EContentScheme::SubmissionPackage) && Owner.Profile->GetAutomatedTests().Num() == 0; }, // run is implied when there are automated tests
 				.Validation = FValidation({ELauncherProfileValidationErrors::ZenPakStreamingRequiresDeployAndLaunch, ELauncherProfileValidationErrors::NoLaunchRoleDeviceAssigned}),
 			}
 		);
 	}
 
 
-	void FGenericProfileTreeBuilder::AddInitialMapProperty( FLaunchProfileTreeNode& HeadingNode )
+	void FGenericProfileTreeBuilder::FBuildCookRun::AddInitialMapProperty( FLaunchProfileTreeNode& HeadingNode )
 	{
 		// todo: custom map picker 
 		HeadingNode.AddWidget( LOCTEXT("InitialMapLabel", "Initial Map"), 
 			{
-				.IsDefault = [this]()		{ return Profile->GetDefaultLaunchRole()->GetInitialMap().IsEmpty(); },
-				.SetToDefault = [this]()	{ Profile->GetDefaultLaunchRole()->SetInitialMap(FString()); InitalMapCombo->SetSelectedItem(GetInitialMap()); },
+				.IsDefault = [this]()		{ return BuildCookRun->GetDefaultLaunchRole()->GetInitialMap().IsEmpty(); },
+				.SetToDefault = [this]()	{ BuildCookRun->GetDefaultLaunchRole()->SetInitialMap(FString()); InitalMapCombo->SetSelectedItem(GetInitialMap()); },
 				.IsVisible = [this]()		{ return (ContentScheme != EContentScheme::SubmissionPackage); },
 				.IsEnabled = ForRun,
 				.Validation = FValidation({ELauncherProfileValidationErrors::InitialMapNotAvailable}),
 			},
 			SAssignNew(InitalMapCombo, SSearchableComboBox)
-			.OptionsSource(&CachedStartupMaps)
-			.OnSelectionChanged( this, &FGenericProfileTreeBuilder::OnInitialMapChanged )
-			.OnGenerateWidget( this, &FGenericProfileTreeBuilder::OnGenerateComboWidget )
-			.OnComboBoxOpening( this, &FGenericProfileTreeBuilder::CacheStartupMapList )
+			.OptionsSource(&Owner.CachedStartupMaps)
+			.OnSelectionChanged( this, &FGenericProfileTreeBuilder::FBuildCookRun::OnInitialMapChanged )
+			.OnGenerateWidget( Owner.AsShared(), &FGenericProfileTreeBuilder::OnGenerateComboWidget )
+			.OnComboBoxOpening( Owner.AsShared(), &FGenericProfileTreeBuilder::CacheStartupMapList )
 			.Content()
 			[
 				SNew(STextBlock)
-				.Text_Lambda( [this]() { return FText::FromString( *Profile->GetDefaultLaunchRole()->GetInitialMap() ); } )
+				.Text_Lambda( [this]() { return FText::FromString( *BuildCookRun->GetDefaultLaunchRole()->GetInitialMap() ); } )
 				.Font(FCoreStyle::Get().GetFontStyle("SmallFont"))
 			]
 		);
 	}
 
 
-	void FGenericProfileTreeBuilder::AddCommandLineProperty( FLaunchProfileTreeNode& HeadingNode )
+	void FGenericProfileTreeBuilder::FBuildCookRun::AddCommandLineProperty( FLaunchProfileTreeNode& HeadingNode )
 	{
 		HeadingNode.AddCommandLineString( LOCTEXT("CommandLineLabel", "Additional Command Line"), 
 			{
@@ -811,28 +836,50 @@ namespace ProjectLauncher
 
 
 
-	void FGenericProfileTreeBuilder::OnValidateProfile()
+	void FGenericProfileTreeBuilder::FBuildCookRun::OnValidateProfile(ILauncherProfileUATCommandRef UATCommand)
 	{
+		if (UATCommand != BuildCookRun)
+		{
+			return;
+		}
+
 		// verify that IoStore is available when a Zen Workflow is in use
 		if ((ContentScheme == EContentScheme::ZenStreaming || ContentScheme == EContentScheme::ZenPakStreaming))
 		{
-			if (!Model->GetProjectSettings(Profile).bUseIoStore)
+			if (!Model->GetProjectSettings(Owner.Profile).bUseIoStore)
 			{
-				Profile->AddCustomError(TEXT("Validation_ZenNeedsIoStore"), LOCTEXT("Error_ZenNeedsIoStore", "Zen Workflows require IoStore to be enabled in your project's Packaging Settings. It cannot be overridden in a profile"));
+				Owner.Profile->AddCustomError(TEXT("Validation_ZenNeedsIoStore"), LOCTEXT("Error_ZenNeedsIoStore", "Zen Workflows require IoStore to be enabled in your project's Packaging Settings. It cannot be overridden in a profile"), BuildCookRun.ToSharedPtr());
 			}
 			
-			if (!Model->GetProjectSettings(Profile).bAllowRemoteNetworkService && Model->IsUsingRemotePlatform(Profile))
+			if (!Model->GetProjectSettings(Owner.Profile).bAllowRemoteNetworkService && Model->IsUsingRemotePlatform(BuildCookRun))
 			{
-				Profile->AddCustomError(TEXT("Validation_ZenNeedsRemotePermission"), LOCTEXT("Error_ZenNeedsRemotePermission", "Zen Workflows for remote devices need [Zen.AutoLaunch]:AllowRemoteNetworkService=true. It cannot be overridden in a profile"));
+				Owner.Profile->AddCustomError(TEXT("Validation_ZenNeedsRemotePermission"), LOCTEXT("Error_ZenNeedsRemotePermission", "Zen Workflows for remote devices need [Zen.AutoLaunch]:AllowRemoteNetworkService=true. It cannot be overridden in a profile"), BuildCookRun.ToSharedPtr());
 			}
 		}
 
 		if (CachedBuildTargetTypes.Contains(EBuildTargetType::Game) && CachedBuildTargetTypes.Contains(EBuildTargetType::Client))
 		{
-			Profile->AddCustomError(TEXT("Validation_GameAndClientTargetsDisallowed"), LOCTEXT("Error_NoGameAndClient", "Cannot specify Game and Client targets to be built together"));
+			Owner.Profile->AddCustomError(TEXT("Validation_GameAndClientTargetsDisallowed"), LOCTEXT("Error_NoGameAndClient", "Cannot specify Game and Client targets to be built together"), BuildCookRun.ToSharedPtr());
+		}
+	}
+
+	void FGenericProfileTreeBuilder::OnUATCommandAdded(const ILauncherProfileUATCommandRef& UATCommand)
+	{
+		ILauncherProfileBuildCookRunPtr BuildCookRun = UATCommand->AsBuildCookRun();
+		if (BuildCookRun.IsValid())
+		{
+			BuildCookRunData.Add(BuildCookRun.ToSharedRef(), MakeShared<FBuildCookRun>(*this, BuildCookRun.ToSharedRef(), Model));
 		}
 	}
 
+	void FGenericProfileTreeBuilder::OnUATCommandRemoved(const ILauncherProfileUATCommandRef& UATCommand)
+	{
+		ILauncherProfileBuildCookRunPtr BuildCookRun = UATCommand->AsBuildCookRun();
+		if (BuildCookRun.IsValid())
+		{
+			BuildCookRunData.Remove(BuildCookRun.ToSharedRef());
+		}
+	}
 
 	void FGenericProfileTreeBuilder::OnPropertyChanged()
 	{
@@ -864,14 +911,9 @@ namespace ProjectLauncher
 			CachedStartupMaps.Add(MakeShared<FString>(Map));
 		}
 
-		if (InitalMapCombo.IsValid())
-		{
-			InitalMapCombo->SetSelectedItem( GetInitialMap(), ESelectInfo::Direct );
-			InitalMapCombo->RefreshOptions();
-		}
 	}
 
-	void FGenericProfileTreeBuilder::CacheArchitectures()
+	void FGenericProfileTreeBuilder::FBuildCookRun::CacheArchitectures()
 	{
 		CachedBuildTargetArchitectures.Reset();
 		CachedBuildTargetArchitectures.Add(EBuildTargetType::Game);
@@ -880,7 +922,7 @@ namespace ProjectLauncher
 		CachedBuildTargetArchitectures.Add(EBuildTargetType::Editor);
 
 		// gather all architecures for each target type
-		for (const PlatformInfo::FTargetPlatformInfo* PlatformInfo : Model->GetPlatformInfos(Profile))
+		for (const PlatformInfo::FTargetPlatformInfo* PlatformInfo : Model->GetPlatformInfos(BuildCookRun))
 		{
 			const ITargetPlatform* TargetPlatform = GetTargetPlatformManager()->FindTargetPlatform(PlatformInfo->IniPlatformName);
 			if (TargetPlatform)
@@ -909,11 +951,11 @@ namespace ProjectLauncher
 		InsertProjectDefault( CachedBuildTargetArchitectures[EBuildTargetType::Editor] );
 	}
 
-	void FGenericProfileTreeBuilder::CacheBuildTargets()
+	void FGenericProfileTreeBuilder::FBuildCookRun::CacheBuildTargets()
 	{
 		CachedBuildTargetTypes.Reset();
 
-		TArray<FTargetInfo> TargetInfos = Model->GetBuildTargetInfos(Profile);
+		TArray<FTargetInfo> TargetInfos = Model->GetBuildTargetInfos(Owner.Profile, BuildCookRun);
 		for (const FTargetInfo& TargetInfo : TargetInfos)
 		{
 			CachedBuildTargetTypes.Add(TargetInfo.Type);
@@ -937,38 +979,38 @@ namespace ProjectLauncher
 
 
 
-	void FGenericProfileTreeBuilder::SetSelectedPlatforms( TArray<FString> SelectedPlatforms )
+	void FGenericProfileTreeBuilder::FBuildCookRun::SetSelectedPlatforms( TArray<FString> SelectedPlatforms )
 	{
 		if (bHasAdvancedPlatformTargets)
 		{
-			Profile->ClearCookedPlatforms();
+			BuildCookRun->ClearCookedPlatforms();
 			for (const FString& Platform : SelectedPlatforms)
 			{
-				Profile->AddCookedPlatform(Platform);
+				BuildCookRun->AddCookedPlatform(Platform);
 			}
 		}
 		else
 		{
-			Profile->ClearCookedPlatforms();
+			BuildCookRun->ClearCookedPlatforms();
 
 			FString BuildTarget = (GetBuildTargets().Num() > 0) ? GetBuildTargets()[0] : FString();
-			FTargetInfo BuildTargetInfo = Model->GetBuildTargetInfo(BuildTarget, GetProjectPath());
+			FTargetInfo BuildTargetInfo = Model->GetBuildTargetInfo(BuildTarget, Owner.GetProjectPath());
 
 			for (const FString& Platform : SelectedPlatforms)
 			{
-				Profile->AddCookedPlatform( FModel::GetBuildTargetPlatformName(Platform, BuildTargetInfo));
+				BuildCookRun->AddCookedPlatform( FModel::GetBuildTargetPlatformName(Platform, BuildTargetInfo));
 			}
 		}
 
-		if (Profile->GetDefaultDeployPlatform() != NAME_None)
+		if (BuildCookRun->GetDefaultDeployPlatform() != NAME_None)
 		{
 			FString DeployPlatform = (GetSelectedPlatforms().Num() > 0) ? GetSelectedPlatforms()[0] : FPlatformProperties::IniPlatformName();
-			Profile->SetDefaultDeployPlatform(FName(*DeployPlatform));
+			BuildCookRun->SetDefaultDeployPlatform(FName(*DeployPlatform));
 		}
 
 		ValidateArchitectures();
 
-		OnPropertyChanged();
+		Owner.OnPropertyChanged();
 
 		if (DeployDeviceListView)
 		{
@@ -978,17 +1020,17 @@ namespace ProjectLauncher
 		CacheArchitectures();
 	}
 
-	TArray<FString> FGenericProfileTreeBuilder::GetSelectedPlatforms() const
+	TArray<FString> FGenericProfileTreeBuilder::FBuildCookRun::GetSelectedPlatforms() const
 	{
 		TArray<FString> Platforms;
 
 		if (bHasAdvancedPlatformTargets)
 		{
-			Platforms = Profile->GetCookedPlatforms();
+			Platforms = BuildCookRun->GetCookedPlatforms();
 		}
 		else
 		{
-			for (const FString& Platform : Profile->GetCookedPlatforms())
+			for (const FString& Platform : BuildCookRun->GetCookedPlatforms())
 			{
 				Platforms.Add(FModel::GetPlatformNameWithFlavor(Platform));
 			}
@@ -1010,7 +1052,10 @@ namespace ProjectLauncher
 
 		if (ProfileType != EProfileType::Basic )
 		{
-			Model->UpdateCookedPlatformsFromBuildTarget(Profile);
+			for (const ILauncherProfileBuildCookRunRef& BuildCookRun : Profile->GetBuildCookRunCommands())
+			{
+				Model->UpdateCookedPlatformsFromBuildTarget(Profile, BuildCookRun);
+			}
 		}
 
 		OnPropertyChanged();
@@ -1024,14 +1069,14 @@ namespace ProjectLauncher
 	}
 
 
-	TArray<FString> FGenericProfileTreeBuilder::GetBuildTargets() const
+	TArray<FString> FGenericProfileTreeBuilder::FBuildCookRun::GetBuildTargets() const
 	{
-		if (ProfileType == EProfileType::Basic)
+		if (Owner.ProfileType == EProfileType::Basic)
 		{
 			check(!bHasAdvancedPlatformTargets);
-			if (Profile->HasBuildTargetSpecified())
+			if (BuildCookRun->HasBuildTargetSpecified())
 			{
-				return Profile->GetBuildTargets();
+				return BuildCookRun->GetBuildTargets();
 			}
 			else
 			{
@@ -1040,24 +1085,24 @@ namespace ProjectLauncher
 		}
 		else
 		{
-			return Profile->GetBuildTargets();
+			return BuildCookRun->GetBuildTargets();
 		}
 	}
 		
-	void FGenericProfileTreeBuilder::SetBuildTargets( TArray<FString> BuildTargets)
+	void FGenericProfileTreeBuilder::FBuildCookRun::SetBuildTargets( TArray<FString> BuildTargets)
 	{
-		Profile->ClearBuildTargets();
+		BuildCookRun->ClearBuildTargets();
 		if (BuildTargets.Num() == 0 || BuildTargets.Contains(FString()))
 		{
-			Profile->SetBuildTargetSpecified(false);
+			BuildCookRun->SetBuildTargetSpecified(false);
 		}
 		else
 		{
-			Profile->SetBuildTargetSpecified(true);
+			BuildCookRun->SetBuildTargetSpecified(true);
 
 			for (const FString& BuildTarget : BuildTargets)
 			{
-				Profile->AddBuildTarget(BuildTarget);
+				BuildCookRun->AddBuildTarget(BuildTarget);
 			}
 		}
 		CacheBuildTargets();
@@ -1065,65 +1110,65 @@ namespace ProjectLauncher
 
 		if (!bHasAdvancedPlatformTargets)
 		{
-			if (ProfileType == EProfileType::Basic )
+			if (Owner.ProfileType == EProfileType::Basic )
 			{
-				Model->UpdatedCookedPlatformsFromDeployDeviceProxy(Profile);
+				Model->UpdatedCookedPlatformsFromDeployDeviceProxy(Owner.Profile, BuildCookRun);
 			}
 			else
 			{
-				Model->UpdateCookedPlatformsFromBuildTarget(Profile);
+				Model->UpdateCookedPlatformsFromBuildTarget(Owner.Profile, BuildCookRun);
 			}
 		}
 
 		ValidateArchitectures();
 
-		OnPropertyChanged();
+		Owner.OnPropertyChanged();
 	}
 
 
-	void FGenericProfileTreeBuilder::SetBuildConfiguration(EBuildConfiguration BuildConfiguration)
+	void FGenericProfileTreeBuilder::FBuildCookRun::SetBuildConfiguration(EBuildConfiguration BuildConfiguration)
 	{
-		Profile->SetBuildConfiguration(BuildConfiguration);
-		OnPropertyChanged();
+		BuildCookRun->SetBuildConfiguration(BuildConfiguration);
+		Owner.OnPropertyChanged();
 	}
 
-	EBuildConfiguration FGenericProfileTreeBuilder::GetBuildConfiguration() const
+	EBuildConfiguration FGenericProfileTreeBuilder::FBuildCookRun::GetBuildConfiguration() const
 	{
-		return Profile->GetBuildConfiguration();
+		return BuildCookRun->GetBuildConfiguration();
 	}
 
 
 
-	void FGenericProfileTreeBuilder::RefreshContentScheme()
+	void FGenericProfileTreeBuilder::FBuildCookRun::RefreshContentScheme()
 	{
-		EContentScheme CurrentContentScheme = Model->DetermineProfileContentScheme(Profile);
+		EContentScheme CurrentContentScheme = Model->DetermineProfileContentScheme(BuildCookRun);
 		SetContentScheme(CurrentContentScheme);
 	}
 
-	void FGenericProfileTreeBuilder::SetContentScheme(EContentScheme InContentScheme)
+	void FGenericProfileTreeBuilder::FBuildCookRun::SetContentScheme(EContentScheme InContentScheme)
 	{
 		ELauncherProfileDeploymentModes::Type DeploymentMode = GetDeployToDevice() ? ELauncherProfileDeploymentModes::CopyToDevice : ELauncherProfileDeploymentModes::DoNotDeploy;
-		Model->SetProfileContentScheme(InContentScheme, Profile, bShouldCook, DeploymentMode );
+		Model->SetProfileContentScheme(InContentScheme, BuildCookRun, bShouldCook, DeploymentMode );
 
 		// refresh the cached content scheme in case the option that was selected is not available
-		ContentScheme = Model->DetermineProfileContentScheme(Profile);
+		ContentScheme = Model->DetermineProfileContentScheme(BuildCookRun);
 
-		OnPropertyChanged();
+		Owner.OnPropertyChanged();
 	}
 
-	bool FGenericProfileTreeBuilder::IsContentSchemeAvailable(EContentScheme InContentScheme, FText& OutReason) const
+	bool FGenericProfileTreeBuilder::FBuildCookRun::IsContentSchemeAvailable(EContentScheme InContentScheme, FText& OutReason) const
 	{
-		FProjectSettings ProjectSettings = Model->GetProjectSettings(Profile);
+		FProjectSettings ProjectSettings = Model->GetProjectSettings(Owner.Profile);
 
 		// basic launch is aimed at launching the current content, not an external build
-		if (ProfileType == EProfileType::Basic && InContentScheme == EContentScheme::ZenPakStreaming)
+		if (Owner.ProfileType == EProfileType::Basic && InContentScheme == EContentScheme::ZenPakStreaming)
 		{
 			// don't set a reason - just hide the item for Basic Launch
 			return false;
 		}
 
 		// cannot create a submission package from the Basic Launch
-		if (ProfileType == EProfileType::Basic && InContentScheme == EContentScheme::SubmissionPackage)
+		if (Owner.ProfileType == EProfileType::Basic && InContentScheme == EContentScheme::SubmissionPackage)
 		{
 			// don't set a reason - just hide the item for Basic Launch
 			return false;
@@ -1146,14 +1191,14 @@ namespace ProjectLauncher
 	
 		// cannot launch via Zen if we're targeting a remote device and Zen isn't accepting external connections
 		if ((InContentScheme == EContentScheme::ZenStreaming || InContentScheme == EContentScheme::ZenPakStreaming) 
-			&& !ProjectSettings.bAllowRemoteNetworkService && Model->IsUsingRemotePlatform(Profile))
+			&& !ProjectSettings.bAllowRemoteNetworkService && Model->IsUsingRemotePlatform(BuildCookRun))
 		{
 			OutReason = LOCTEXT("NoZenReason", "Zen Streaming to a remote device requires AllowRemoteNetworkService");
 			return false;
 		}
 
 		// cannot make a submission package outside of the shipping configuration
-		if (InContentScheme == EContentScheme::SubmissionPackage && Profile->GetBuildConfiguration() != EBuildConfiguration::Shipping)
+		if (InContentScheme == EContentScheme::SubmissionPackage && BuildCookRun->GetBuildConfiguration() != EBuildConfiguration::Shipping)
 		{
 			OutReason = LOCTEXT("SubPkgShippingOnly", "Submission packages require a Shipping build");
 			return false;
@@ -1164,35 +1209,35 @@ namespace ProjectLauncher
 
 
 
-	FString FGenericProfileTreeBuilder::GetCommandLine(EBuildTargetType BuildTarget) const
+	FString FGenericProfileTreeBuilder::FBuildCookRun::GetCommandLine(EBuildTargetType BuildTarget) const
 	{
 		if (BuildTarget == EBuildTargetType::Unknown)
 		{
 			// get unified command line from these two fields. the first is presented in "Build" for old Project Launcher, and the latter is presented in "Launch" for old Project Launcher)
 			// when we save back to the profile, this will be stored just in the "Build" one for clarity (because multiple roles are not supported in old or new project launcher)
-			FString CommandLine = Profile->GetAdditionalCommandLineParameters().TrimStartAndEnd() + TEXT(" ") + Profile->GetDefaultLaunchRole()->GetUATCommandLine().TrimStartAndEnd();
+			FString CommandLine = BuildCookRun->GetAdditionalCommandLineParameters().TrimStartAndEnd() + TEXT(" ") + BuildCookRun->GetDefaultLaunchRole()->GetUATCommandLine().TrimStartAndEnd();
 			CommandLine.TrimStartAndEndInline();
 
 			return CommandLine;
 		}
 		else
 		{
-			return Profile->GetAdditionalTargetCommandLineParameters(BuildTarget).TrimStartAndEnd();
+			return BuildCookRun->GetAdditionalTargetCommandLineParameters(BuildTarget).TrimStartAndEnd();
 		}
 	}
 
-	void FGenericProfileTreeBuilder::SetCommandLine( const FString& NewCommandLine, EBuildTargetType BuildTarget )
+	void FGenericProfileTreeBuilder::FBuildCookRun::SetCommandLine( const FString& NewCommandLine, EBuildTargetType BuildTarget )
 	{
 		if (BuildTarget == EBuildTargetType::Unknown)
 		{
-			Profile->SetAdditionalCommandLineParameters(NewCommandLine);
-			Profile->GetDefaultLaunchRole()->SetCommandLine(TEXT(""));
+			BuildCookRun->SetAdditionalCommandLineParameters(NewCommandLine);
+			BuildCookRun->GetDefaultLaunchRole()->SetCommandLine(TEXT(""));
 		}
 		else
 		{
-			Profile->SetAdditionalTargetCommandLineParameters(NewCommandLine, BuildTarget);
+			BuildCookRun->SetAdditionalTargetCommandLineParameters(NewCommandLine, BuildTarget);
 		}
-		OnPropertyChanged();
+		Owner.OnPropertyChanged();
 	}
 
 
@@ -1201,19 +1246,19 @@ namespace ProjectLauncher
 
 
 
-	void FGenericProfileTreeBuilder::SetCook( bool bCook )
+	void FGenericProfileTreeBuilder::FBuildCookRun::SetCook( bool bCook )
 	{
 		bShouldCook = bCook;
 		RefreshContentScheme();
 	}
 
-	bool FGenericProfileTreeBuilder::GetCook(ILauncherProfilePtr InProfile) const
+	bool FGenericProfileTreeBuilder::FBuildCookRun::GetCook(ILauncherProfileBuildCookRunPtr InBuildCookRun) const
 	{
-		if (InProfile == nullptr)
+		if (InBuildCookRun == nullptr)
 		{
-			InProfile = Profile;
+			InBuildCookRun = BuildCookRun;
 		}
-		return (InProfile->GetCookMode() != ELauncherProfileCookModes::DoNotCook);
+		return (InBuildCookRun->GetCookMode() != ELauncherProfileCookModes::DoNotCook);
 	}
 
 
@@ -1222,13 +1267,13 @@ namespace ProjectLauncher
 
 
 	BEGIN_SLATE_FUNCTION_BUILD_OPTIMIZATION
-	TSharedRef<SWidget> FGenericProfileTreeBuilder::CreateDeployDeviceWidget()
+	TSharedRef<SWidget> FGenericProfileTreeBuilder::FBuildCookRun::CreateDeployDeviceWidget()
 	{
-		if (ProfileType == EProfileType::Basic)
+		if (Owner.ProfileType == EProfileType::Basic)
 		{
 			return SNew(SCustomLaunchDeviceListView)
-				.OnSelectionChanged(this, &FGenericProfileTreeBuilder::SetDeployDeviceIDs)
-				.SelectedDevices(this, &FGenericProfileTreeBuilder::GetDeployDeviceIDs)
+				.OnSelectionChanged(this, &FGenericProfileTreeBuilder::FBuildCookRun::SetDeployDeviceIDs)
+				.SelectedDevices(this, &FGenericProfileTreeBuilder::FBuildCookRun::GetDeployDeviceIDs)
 				.AllPlatforms(true)
 				.SingleSelect(true)
 			;
@@ -1250,8 +1295,8 @@ namespace ProjectLauncher
 					.VAlign(VAlign_Center)
 					[
 						SNew(SSegmentedControl<EDeployDeviceOption>)
-						.Value( this, &FGenericProfileTreeBuilder::GetDeployDeviceOption)
-						.OnValueChanged(this, &FGenericProfileTreeBuilder::SetDeployDeviceOption)
+						.Value( this, &FGenericProfileTreeBuilder::FBuildCookRun::GetDeployDeviceOption)
+						.OnValueChanged(this, &FGenericProfileTreeBuilder::FBuildCookRun::SetDeployDeviceOption)
 
 						+SSegmentedControl<EDeployDeviceOption>::Slot(EDeployDeviceOption::Default)
 						.HAlign(HAlign_Center)
@@ -1280,15 +1325,15 @@ namespace ProjectLauncher
 					SNew(SVerticalResizeBox)
 					.Visibility_Lambda([this](){ return (GetDeployDeviceOption() == EDeployDeviceOption::Selected) ? EVisibility::Visible : EVisibility::Collapsed; } )
 					.HandleHeight(4)
-					.ContentHeight(this, &FGenericProfileTreeBuilder::GetDeployDeviceListHeight)
-					.ContentHeightChanged(this, &FGenericProfileTreeBuilder::SetDeployDeviceListHeight)
+					.ContentHeight(this, &FGenericProfileTreeBuilder::FBuildCookRun::GetDeployDeviceListHeight)
+					.ContentHeightChanged(this, &FGenericProfileTreeBuilder::FBuildCookRun::SetDeployDeviceListHeight)
 					.HandleColor( FAppStyle::Get().GetSlateColor("Colors.Secondary").GetSpecifiedColor() )
 					[
 						SAssignNew(DeployDeviceListView, SCustomLaunchDeviceListView)
-						.OnDeviceRemoved( this, &FGenericProfileTreeBuilder::OnDeviceRemoved)
-						.OnSelectionChanged(this, &FGenericProfileTreeBuilder::SetDeployDeviceIDs)
-						.SelectedDevices(this, &FGenericProfileTreeBuilder::GetDeployDeviceIDs)
-						.Platforms(this, &FGenericProfileTreeBuilder::GetSelectedPlatforms)
+						.OnDeviceRemoved( this, &FGenericProfileTreeBuilder::FBuildCookRun::OnDeviceRemoved)
+						.OnSelectionChanged(this, &FGenericProfileTreeBuilder::FBuildCookRun::SetDeployDeviceIDs)
+						.SelectedDevices(this, &FGenericProfileTreeBuilder::FBuildCookRun::GetDeployDeviceIDs)
+						.Platforms(this, &FGenericProfileTreeBuilder::FBuildCookRun::GetSelectedPlatforms)
 					]
 				]
 			;
@@ -1297,82 +1342,82 @@ namespace ProjectLauncher
 	END_SLATE_FUNCTION_BUILD_OPTIMIZATION
 
 
-	void FGenericProfileTreeBuilder::SetDeployDeviceIDs(TArray<FString> DeployDeviceIDs)
+	void FGenericProfileTreeBuilder::FBuildCookRun::SetDeployDeviceIDs(TArray<FString> DeployDeviceIDs)
 	{
-		if (Profile->GetDeployedDeviceGroup() == nullptr)
+		if (BuildCookRun->GetDeployedDeviceGroup() == nullptr)
 		{
-			Profile->SetDeployedDeviceGroup(Model->GetProfileManager()->AddNewDeviceGroup());
+			BuildCookRun->SetDeployedDeviceGroup(Model->GetProfileManager()->AddNewDeviceGroup());
 		}
 
-		Profile->GetDeployedDeviceGroup()->RemoveAllDevices(); 
+		BuildCookRun->GetDeployedDeviceGroup()->RemoveAllDevices(); 
 		for ( const FString& DeviceID : DeployDeviceIDs )
 		{
-			Profile->GetDeployedDeviceGroup()->AddDevice(DeviceID);
+			BuildCookRun->GetDeployedDeviceGroup()->AddDevice(DeviceID);
 		}
 
-		if (ProfileType == EProfileType::Basic)
+		if (Owner.ProfileType == EProfileType::Basic)
 		{
-			Model->UpdatedCookedPlatformsFromDeployDeviceProxy(Profile);
+			Model->UpdatedCookedPlatformsFromDeployDeviceProxy(Owner.Profile, BuildCookRun);
 		}
 
-		OnPropertyChanged();
+		Owner.OnPropertyChanged();
 	}
 
-	TArray<FString> FGenericProfileTreeBuilder::GetDeployDeviceIDs() const
+	TArray<FString> FGenericProfileTreeBuilder::FBuildCookRun::GetDeployDeviceIDs() const
 	{
-		return Profile->GetDeployedDeviceGroup()->GetDeviceIDs();
+		return BuildCookRun->GetDeployedDeviceGroup()->GetDeviceIDs();
 	}
 
-	float FGenericProfileTreeBuilder::GetDeployDeviceListHeight() const
+	float FGenericProfileTreeBuilder::FBuildCookRun::GetDeployDeviceListHeight() const
 	{
 		return DeployDeviceListHeight;
 	}
 
-	void FGenericProfileTreeBuilder::SetDeployDeviceListHeight( float NewHeight )
+	void FGenericProfileTreeBuilder::FBuildCookRun::SetDeployDeviceListHeight( float NewHeight )
 	{
 		static const float MinDeployDeviceListHeight = 100.0f;
 
 		DeployDeviceListHeight = FMath::Max(NewHeight, MinDeployDeviceListHeight);
 
-		TreeData->RequestTreeRefresh();
+		Owner.TreeData->RequestTreeRefresh();
 	}
 
-	void FGenericProfileTreeBuilder::OnDeviceRemoved(FString DeviceID)
+	void FGenericProfileTreeBuilder::FBuildCookRun::OnDeviceRemoved(FString DeviceID)
 	{
 		CachedDeployDeviceIDs.Remove(DeviceID);
 	}
 
-	FGenericProfileTreeBuilder::EDeployDeviceOption FGenericProfileTreeBuilder::GetDeployDeviceOption() const
+	FGenericProfileTreeBuilder::FBuildCookRun::EDeployDeviceOption FGenericProfileTreeBuilder::FBuildCookRun::GetDeployDeviceOption() const
 	{
-		return (Profile->GetDefaultDeployPlatform() == NAME_None) ? EDeployDeviceOption::Selected : EDeployDeviceOption::Default;;
+		return (BuildCookRun->GetDefaultDeployPlatform() == NAME_None) ? EDeployDeviceOption::Selected : EDeployDeviceOption::Default;;
 	}
 
-	void FGenericProfileTreeBuilder::SetDeployDeviceOption( EDeployDeviceOption NewDeployDeviceOption )
+	void FGenericProfileTreeBuilder::FBuildCookRun::SetDeployDeviceOption( EDeployDeviceOption NewDeployDeviceOption )
 	{
 		bool bShow = (NewDeployDeviceOption == EDeployDeviceOption::Selected);
 
 		if (bShow)
 		{
-			Profile->SetDefaultDeployPlatform(NAME_None);
+			BuildCookRun->SetDefaultDeployPlatform(NAME_None);
 
 			// restore the deployed device list again, if any
-			if (CachedDeployDeviceIDs.Num() > 0 && Profile->GetDeployedDeviceGroup()->GetDeviceIDs().Num() == 0)
+			if (CachedDeployDeviceIDs.Num() > 0 && BuildCookRun->GetDeployedDeviceGroup()->GetDeviceIDs().Num() == 0)
 			{
-				SetDeployDeviceIDs(Profile->GetDeployedDeviceGroup()->GetDeviceIDs());
+				SetDeployDeviceIDs(BuildCookRun->GetDeployedDeviceGroup()->GetDeviceIDs());
 				CachedDeployDeviceIDs.Reset();
 			}
 		}
 		else
 		{
 			// to set the 'default' deploy option, its necessary to remove all the devices - take a copy of the values to allow it to be restored
-			CachedDeployDeviceIDs = Profile->GetDeployedDeviceGroup()->GetDeviceIDs();
+			CachedDeployDeviceIDs = BuildCookRun->GetDeployedDeviceGroup()->GetDeviceIDs();
 			SetDeployDeviceIDs(TArray<FString>());
 
 			FString DeployPlatform = (GetSelectedPlatforms().Num() > 0) ? GetSelectedPlatforms()[0] : FPlatformProperties::IniPlatformName();
-			Profile->SetDefaultDeployPlatform(FName(*DeployPlatform));
+			BuildCookRun->SetDefaultDeployPlatform(FName(*DeployPlatform));
 		}
 
-		OnPropertyChanged();
+		Owner.OnPropertyChanged();
 
 		if (DeployDeviceListView)
 		{
@@ -1381,61 +1426,61 @@ namespace ProjectLauncher
 	}
 
 
-	void FGenericProfileTreeBuilder::SetBuild( bool bBuild )
+	void FGenericProfileTreeBuilder::FBuildCookRun::SetBuild( bool bBuild )
 	{
 		if (!bBuild)
 		{
-			Profile->SetBuildMode(ELauncherProfileBuildModes::DoNotBuild);
+			BuildCookRun->SetBuildMode(ELauncherProfileBuildModes::DoNotBuild);
 		}
-		else if (GetForceBuild(Profile))
+		else if (GetForceBuild(BuildCookRun))
 		{
-			Profile->SetBuildMode(ELauncherProfileBuildModes::Build);
+			BuildCookRun->SetBuildMode(ELauncherProfileBuildModes::Build);
 		}
 		else
 		{
-			Profile->SetBuildMode(ELauncherProfileBuildModes::Auto);
+			BuildCookRun->SetBuildMode(ELauncherProfileBuildModes::Auto);
 		}
 
-		OnPropertyChanged();
+		Owner.OnPropertyChanged();
 	}
 
-	bool FGenericProfileTreeBuilder::GetBuild(ILauncherProfilePtr InProfile) const
+	bool FGenericProfileTreeBuilder::FBuildCookRun::GetBuild(ILauncherProfileBuildCookRunPtr InBuildCookRun) const
 	{
-		if (InProfile == nullptr)
+		if (InBuildCookRun == nullptr)
 		{
-			InProfile = Profile;
+			InBuildCookRun = BuildCookRun;
 		}
-		return (InProfile->GetBuildMode() != ELauncherProfileBuildModes::DoNotBuild);
+		return (InBuildCookRun->GetBuildMode() != ELauncherProfileBuildModes::DoNotBuild);
 	}
 
-	void FGenericProfileTreeBuilder::SetForceBuild( bool bForceBuild )
+	void FGenericProfileTreeBuilder::FBuildCookRun::SetForceBuild( bool bForceBuild )
 	{
 		if (!GetBuild())
 		{
-			Profile->SetBuildMode(ELauncherProfileBuildModes::DoNotBuild);
+			BuildCookRun->SetBuildMode(ELauncherProfileBuildModes::DoNotBuild);
 		}
 		else if (bForceBuild)
 		{
-			Profile->SetBuildMode(ELauncherProfileBuildModes::Build);
+			BuildCookRun->SetBuildMode(ELauncherProfileBuildModes::Build);
 		}
 		else
 		{
-			Profile->SetBuildMode(ELauncherProfileBuildModes::Auto);
+			BuildCookRun->SetBuildMode(ELauncherProfileBuildModes::Auto);
 		}
 
-		OnPropertyChanged();
+		Owner.OnPropertyChanged();
 	}
 
-	bool FGenericProfileTreeBuilder::GetForceBuild(ILauncherProfilePtr InProfile) const
+	bool FGenericProfileTreeBuilder::FBuildCookRun::GetForceBuild(ILauncherProfileBuildCookRunPtr InBuildCookRun) const
 	{
-		if (InProfile == nullptr)
+		if (InBuildCookRun == nullptr)
 		{
-			InProfile = Profile;
+			InBuildCookRun = BuildCookRun;
 		}
-		return (InProfile->GetBuildMode() == ELauncherProfileBuildModes::Build);
+		return (InBuildCookRun->GetBuildMode() == ELauncherProfileBuildModes::Build);
 	}
 
-	void FGenericProfileTreeBuilder::SetArchitecture( FString Architecture, EBuildTargetType BuildTargetType )
+	void FGenericProfileTreeBuilder::FBuildCookRun::SetArchitecture( FString Architecture, EBuildTargetType BuildTargetType )
 	{
 		// Project Launcher only supports launching for one architecture at a time at the moment
 		TArray<FString> Architectures;
@@ -1447,36 +1492,36 @@ namespace ProjectLauncher
 		// update the architecture
 		if (BuildTargetType == EBuildTargetType::Game || BuildTargetType == EBuildTargetType::Client)
 		{
-			Profile->SetClientArchitectures(Architectures);
+			BuildCookRun->SetClientArchitectures(Architectures);
 		}
 		else if (BuildTargetType == EBuildTargetType::Server)
 		{
-			Profile->SetServerArchitectures(Architectures);
+			BuildCookRun->SetServerArchitectures(Architectures);
 		}
 		else if (BuildTargetType == EBuildTargetType::Editor)
 		{
-			Profile->SetEditorArchitectures(Architectures);
+			BuildCookRun->SetEditorArchitectures(Architectures);
 		}
 
-		OnPropertyChanged();
+		Owner.OnPropertyChanged();
 	}
 
-	FString FGenericProfileTreeBuilder::GetArchitecture(EBuildTargetType BuildTargetType) const
+	FString FGenericProfileTreeBuilder::FBuildCookRun::GetArchitecture(EBuildTargetType BuildTargetType) const
 	{
 		// read current architectures
 		TArray<FString> Architectures;
 		switch (BuildTargetType)
 		{
-			case EBuildTargetType::Server: Architectures = Profile->GetServerArchitectures(); break;
-			case EBuildTargetType::Editor: Architectures = Profile->GetEditorArchitectures(); break;
-			default:                       Architectures = Profile->GetClientArchitectures(); break;
+			case EBuildTargetType::Server: Architectures = BuildCookRun->GetServerArchitectures(); break;
+			case EBuildTargetType::Editor: Architectures = BuildCookRun->GetEditorArchitectures(); break;
+			default:                       Architectures = BuildCookRun->GetClientArchitectures(); break;
 		}
 
 		// Project Launcher only supports launching for one architecture at a time at the moment
 		return (Architectures.Num() > 0) ? Architectures[0] : FString();
 	}
 
-	FText FGenericProfileTreeBuilder::GetArchitectureDisplayName( FString Architecture )
+	FText FGenericProfileTreeBuilder::FBuildCookRun::GetArchitectureDisplayName( FString Architecture )
 	{
 		if (Architecture.IsEmpty())
 		{
@@ -1492,7 +1537,7 @@ namespace ProjectLauncher
 		}
 	}
 
-	void FGenericProfileTreeBuilder::ValidateArchitectures()
+	void FGenericProfileTreeBuilder::FBuildCookRun::ValidateArchitectures()
 	{
 		CacheArchitectures();
 
@@ -1500,73 +1545,73 @@ namespace ProjectLauncher
 
 		if (CachedBuildTargetArchitectures[EBuildTargetType::Game].Num() == 0 && CachedBuildTargetArchitectures[EBuildTargetType::Client].Num() == 0)
 		{
-			Profile->SetClientArchitectures(TArray<FString>());
+			BuildCookRun->SetClientArchitectures(TArray<FString>());
 		}
 		if (CachedBuildTargetArchitectures[EBuildTargetType::Server].Num() == 0)
 		{
-			Profile->SetServerArchitectures(TArray<FString>());
+			BuildCookRun->SetServerArchitectures(TArray<FString>());
 		}
 		if (CachedBuildTargetArchitectures[EBuildTargetType::Editor].Num() == 0)
 		{
-			Profile->SetEditorArchitectures(TArray<FString>());
+			BuildCookRun->SetEditorArchitectures(TArray<FString>());
 		}
 	}
 
 
-	void FGenericProfileTreeBuilder::SetDeployToDevice( bool bDeployToDevice )
+	void FGenericProfileTreeBuilder::FBuildCookRun::SetDeployToDevice( bool bDeployToDevice )
 	{
 		if (ContentScheme != EContentScheme::CookOnTheFly)
 		{
-			Profile->SetDeploymentMode(bDeployToDevice ? ELauncherProfileDeploymentModes::CopyToDevice : ELauncherProfileDeploymentModes::DoNotDeploy );
-			OnPropertyChanged();
+			BuildCookRun->SetDeploymentMode(bDeployToDevice ? ELauncherProfileDeploymentModes::CopyToDevice : ELauncherProfileDeploymentModes::DoNotDeploy );
+			Owner.OnPropertyChanged();
 		}
 	}
 
-	bool FGenericProfileTreeBuilder::GetDeployToDevice(ILauncherProfilePtr InProfile) const
+	bool FGenericProfileTreeBuilder::FBuildCookRun::GetDeployToDevice(ILauncherProfileBuildCookRunPtr InBuildCookRun) const
 	{
-		if (InProfile == nullptr)
+		if (InBuildCookRun == nullptr)
 		{
-			InProfile = Profile;
+			InBuildCookRun = BuildCookRun;
 		}
-		return InProfile->GetDeploymentMode() != ELauncherProfileDeploymentModes::DoNotDeploy;
+		return InBuildCookRun->GetDeploymentMode() != ELauncherProfileDeploymentModes::DoNotDeploy;
 	}
 
 
-	void FGenericProfileTreeBuilder::SetIsRunning( bool bRun )
+	void FGenericProfileTreeBuilder::FBuildCookRun::SetIsRunning( bool bRun )
 	{
-		Profile->SetLaunchMode(bRun ? ELauncherProfileLaunchModes::DefaultRole : ELauncherProfileLaunchModes::DoNotLaunch);
+		BuildCookRun->SetLaunchMode(bRun ? ELauncherProfileLaunchModes::DefaultRole : ELauncherProfileLaunchModes::DoNotLaunch);
 	}
 
-	bool FGenericProfileTreeBuilder::GetIsRunning( ILauncherProfilePtr InProfile) const
+	bool FGenericProfileTreeBuilder::FBuildCookRun::GetIsRunning( ILauncherProfileBuildCookRunPtr InBuildCookRun) const
 	{
-		if (InProfile == nullptr)
+		if (InBuildCookRun == nullptr)
 		{
-			InProfile = Profile;
+			InBuildCookRun = BuildCookRun;
 		}
 	
-		return InProfile->GetLaunchMode() != ELauncherProfileLaunchModes::DoNotLaunch;
+		return InBuildCookRun->GetLaunchMode() != ELauncherProfileLaunchModes::DoNotLaunch;
 	}
 
-	void FGenericProfileTreeBuilder::OnInitialMapChanged( TSharedPtr<FString> InitialMap, ESelectInfo::Type )
+	void FGenericProfileTreeBuilder::FBuildCookRun::OnInitialMapChanged( TSharedPtr<FString> InitialMap, ESelectInfo::Type )
 	{
 		if (InitialMap.IsValid())
 		{
-			Profile->GetDefaultLaunchRole()->SetInitialMap(*InitialMap);
+			BuildCookRun->GetDefaultLaunchRole()->SetInitialMap(*InitialMap);
 		}
 		else
 		{
-			Profile->GetDefaultLaunchRole()->SetInitialMap(FString());
+			BuildCookRun->GetDefaultLaunchRole()->SetInitialMap(FString());
 		}
 
-		OnPropertyChanged();
+		Owner.OnPropertyChanged();
 	}
 
-	TSharedPtr<FString> FGenericProfileTreeBuilder::GetInitialMap() const
+	TSharedPtr<FString> FGenericProfileTreeBuilder::FBuildCookRun::GetInitialMap() const
 	{
-		CacheStartupMapList();
+		Owner.CacheStartupMapList();
 
-		FString InitialMap = Profile->GetDefaultLaunchRole()->GetInitialMap();
-		const TSharedPtr<FString>* FoundMapPtr = CachedStartupMaps.FindByPredicate( [InitialMap]( const TSharedPtr<FString>& Map )
+		FString InitialMap = BuildCookRun->GetDefaultLaunchRole()->GetInitialMap();
+		const TSharedPtr<FString>* FoundMapPtr = Owner.CachedStartupMaps.FindByPredicate( [InitialMap]( const TSharedPtr<FString>& Map )
 		{
 			return Map.IsValid() && (InitialMap == *Map);
 		});
@@ -1582,11 +1627,11 @@ namespace ProjectLauncher
 	}
 
 
-	void FGenericProfileTreeBuilder::SetHasAdvancedPlatformTargets( bool bEnable )
+	void FGenericProfileTreeBuilder::FBuildCookRun::SetHasAdvancedPlatformTargets( bool bEnable )
 	{
 		if (bHasAdvancedPlatformTargets && !bEnable)
 		{
-			if (Profile->GetCookedPlatforms().Num() > 1 || Profile->GetBuildTargets().Num() > 1)
+			if (BuildCookRun->GetCookedPlatforms().Num() > 1 || BuildCookRun->GetBuildTargets().Num() > 1)
 			{
 				FText MessageText = LOCTEXT("ClearAdvancedPlatformTargets", "You have multiple platform/targets selected - are you sure you want to go back to single platform/target??");
 				TSharedRef<SMessageDialog> Dialog = SNew(SMessageDialog)
@@ -1602,18 +1647,18 @@ namespace ProjectLauncher
 
 			bHasAdvancedPlatformTargets = bEnable;
 
-			if (Profile->GetBuildTargets().Num() > 1)
+			if (BuildCookRun->GetBuildTargets().Num() > 1)
 			{
-				FString BuildTarget = Profile->GetBuildTargets()[0];
-				Profile->ClearBuildTargets();
-				Profile->AddBuildTarget(BuildTarget);
+				FString BuildTarget = BuildCookRun->GetBuildTargets()[0];
+				BuildCookRun->ClearBuildTargets();
+				BuildCookRun->AddBuildTarget(BuildTarget);
 
 			}
-			Model->UpdateCookedPlatformsFromBuildTarget(Profile);
+			Model->UpdateCookedPlatformsFromBuildTarget(Owner.Profile, BuildCookRun);
 
-			if (Profile->GetCookedPlatforms().Num() > 0)
+			if (BuildCookRun->GetCookedPlatforms().Num() > 0)
 			{
-				TArray<FString> Platforms{Profile->GetCookedPlatforms()[0]};
+				TArray<FString> Platforms{BuildCookRun->GetCookedPlatforms()[0]};
 				SetSelectedPlatforms(Platforms);
 			}
 
@@ -1622,7 +1667,7 @@ namespace ProjectLauncher
 		else
 		{
 			bHasAdvancedPlatformTargets = bEnable;
-			OnPropertyChanged();
+			Owner.OnPropertyChanged();
 		}
 
 		if (PlatformCombo)
@@ -1631,12 +1676,17 @@ namespace ProjectLauncher
 		}
 	}
 
-	bool FGenericProfileTreeBuilder::HasAdvancedPlatformTargets() const
+	bool FGenericProfileTreeBuilder::FBuildCookRun::HasAdvancedPlatformTargets() const
 	{
 		return bHasAdvancedPlatformTargets;
 	}
 
 
+	FGenericProfileTreeBuilder::FBuildCookRun& FGenericProfileTreeBuilder::Get(const ILauncherProfileBuildCookRunRef& BuildCookRun)
+	{
+		return BuildCookRunData.FindChecked(BuildCookRun).Get();
+	}
+
 }
 
 #undef LOCTEXT_NAMESPACE
diff --git a/Engine/Plugins/Developer/ProjectLauncher/Source/ProjectLauncher/Private/ProfileTree/LaunchProfileTreeData.cpp b/Engine/Plugins/Developer/ProjectLauncher/Source/ProjectLauncher/Private/ProfileTree/LaunchProfileTreeData.cpp
--- a/Engine/Plugins/Developer/ProjectLauncher/Source/ProjectLauncher/Private/ProfileTree/LaunchProfileTreeData.cpp
+++ b/Engine/Plugins/Developer/ProjectLauncher/Source/ProjectLauncher/Private/ProfileTree/LaunchProfileTreeData.cpp
@@ -51,13 +51,13 @@ namespace ProjectLauncher
 		return bIsSet;
 	}
 
-	bool FValidation::HasError( ILauncherProfileRef Profile ) const
+	bool FValidation::HasError( ILauncherProfileRef Profile, ILauncherProfileUATCommandPtr UATCommand ) const
 	{
-		if (Profile->HasValidationError())
+		if (Profile->HasValidationError(UATCommand))
 		{
 			for (ELauncherProfileValidationErrors::Type Error : AssociatedErrors)
 			{
-				if (Profile->HasValidationError(Error))
+				if (Profile->HasValidationError(Error, UATCommand))
 				{
 					return true;
 				}
@@ -66,7 +66,7 @@ namespace ProjectLauncher
 
 		for (const FString& CustomError : AssociatedCustomErrors)
 		{
-			if (Profile->HasCustomError(CustomError))
+			if (Profile->HasCustomError(CustomError, UATCommand))
 			{
 				return true;
 			}
@@ -74,7 +74,7 @@ namespace ProjectLauncher
 
 		for (const FString& CustomError : AssociatedCustomErrors)
 		{
-			if (Profile->HasCustomWarning(CustomError))
+			if (Profile->HasCustomWarning(CustomError, UATCommand))
 			{
 				return true;
 			}
@@ -83,11 +83,11 @@ namespace ProjectLauncher
 		return false;
 	}
 
-	void FValidation::GetErrorText( ILauncherProfileRef Profile, TArray<FString>& ErrorLines ) const
+	void FValidation::GetErrorText( ILauncherProfileRef Profile, ILauncherProfileUATCommandPtr UATCommand, TArray<FString>& ErrorLines ) const
 	{
 		for (ELauncherProfileValidationErrors::Type Error : AssociatedErrors)
 		{
-			if (Profile->HasValidationError(Error))
+			if (Profile->HasValidationError(Error, UATCommand))
 			{
 				ErrorLines.AddUnique(LexToStringLocalized(Error));
 			}
@@ -95,17 +95,17 @@ namespace ProjectLauncher
 
 		for (const FString& CustomError : AssociatedCustomErrors)
 		{
-			if (Profile->HasCustomError(CustomError))
+			if (Profile->HasCustomError(CustomError, UATCommand))
 			{
-				ErrorLines.AddUnique(Profile->GetCustomErrorText(CustomError).ToString());
+				ErrorLines.AddUnique(Profile->GetCustomErrorText(CustomError, UATCommand).ToString());
 			}
 		}
 
 		for (const FString& CustomWarning : AssociatedCustomErrors)
 		{
-			if (Profile->HasCustomWarning(CustomWarning))
+			if (Profile->HasCustomWarning(CustomWarning, UATCommand))
 			{
-				ErrorLines.AddUnique(Profile->GetCustomWarningText(CustomWarning).ToString());
+				ErrorLines.AddUnique(Profile->GetCustomWarningText(CustomWarning, UATCommand).ToString());
 			}
 		}
 	}
@@ -123,7 +123,7 @@ namespace ProjectLauncher
 		{
 			ExtensionInstances = FLaunchExtension::CreateExtensionInstancesForProfile(Profile.ToSharedRef(), InModel, *this);
 
-			for (TSharedPtr<FLaunchExtensionInstance> ExtensionInstance : ExtensionInstances)
+			for (const TSharedPtr<FLaunchExtensionInstance>& ExtensionInstance : ExtensionInstances)
 			{
 				ExtensionInstance->GetPropertyChangedDelegate().AddRaw(TreeBuilder, &ILaunchProfileTreeBuilder::OnPropertyChanged );
 			}
@@ -132,10 +132,8 @@ namespace ProjectLauncher
 
 	FLaunchProfileTreeData::~FLaunchProfileTreeData()
 	{
-		for (TSharedPtr<FLaunchExtensionInstance> ExtensionInstance : ExtensionInstances)
-		{
-			ExtensionInstance->GetPropertyChangedDelegate().RemoveAll(TreeBuilder);
-		}
+		FLaunchExtension::DestroyExtensionInstancesForProfile(Profile);
+		ExtensionInstances.Reset();
 	}
 
 
@@ -153,12 +151,10 @@ namespace ProjectLauncher
 		Nodes.Add(TreeNode);
 		RequestExpandHeadings.Add(TreeNode);
 
+
 		HeadingNodes.Add(InName, TreeNode);
 
-		Nodes.StableSort( []( const FLaunchProfileTreeNodePtr& A, const FLaunchProfileTreeNodePtr& B)
-		{
-			return A->SortOrder < B->SortOrder;
-		});
+		SortNodes();
 
 		return *TreeNode.Get();
 	}
@@ -169,11 +165,37 @@ namespace ProjectLauncher
 		return *TreeNode.Get();
 	}
 
+	FLaunchProfileTreeNode& FLaunchProfileTreeData::FindOrAddHeading(ILauncherProfileUATCommandRef InUATCommand, bool bIsPrimaryHeading)
+	{
+		if (HeadingNodes.Contains(InUATCommand->GetInternalName()))
+		{
+			return GetHeading(InUATCommand->GetInternalName());
+		}
+
+		FLaunchProfileTreeNode& TreeNode = AddHeading(InUATCommand->GetInternalName(), FText::FromString(*InUATCommand->GetDescription()), InUATCommand->GetOrder()); // @todo: need description to auto-update
+		TreeNode.UATCommand = InUATCommand;
+
+		if (bIsPrimaryHeading)
+		{
+			UATCommandHeadings.Add(TreeNode.AsShared());
+		}
+
+		return TreeNode;
+	}
+
+	void FLaunchProfileTreeData::SortNodes()
+	{
+		Nodes.StableSort( []( const FLaunchProfileTreeNodePtr& A, const FLaunchProfileTreeNodePtr& B)
+		{
+			return A->SortOrder < B->SortOrder;
+		});
+	}
+
 	void FLaunchProfileTreeData::CreateExtensionsUI()
 	{
 		for (TSharedPtr<FLaunchExtensionInstance> ExtensionInstance : ExtensionInstances)
 		{
-			ExtensionInstance->CustomizeTree(*this);
+			ExtensionInstance->CustomizeTree();
 		}
 	}
 	
@@ -210,7 +232,23 @@ namespace ProjectLauncher
 		TreeData->TreeBuilder->OnPropertyChanged();
 	}
 
+	bool FLaunchProfileTreeNode::ContainsNodeRecursive( const FLaunchProfileTreeNodePtr& TargetNode ) const
+	{
+		if (TargetNode.Get() == this)
+		{
+			return true;
+		}
 
+		for (const ProjectLauncher::FLaunchProfileTreeNodePtr& Child : Children)
+		{
+			if (Child->ContainsNodeRecursive(TargetNode))
+			{
+				return true;
+			}
+		}
+
+		return false;
+	}
 
 
 	FLaunchProfileTreeNode& FLaunchProfileTreeNode::AddSubHeading( const TCHAR* InName, FText InDisplayName, int32 InSortOrder )
@@ -224,6 +262,7 @@ namespace ProjectLauncher
 		TreeNode->Name = InDisplayName;
 		TreeNode->SortOrder = InSortOrder;
 		TreeNode->Depth = Depth+1;
+		TreeNode->UATCommand = UATCommand;
 		Children.Add(TreeNode);
 		TreeData->RequestExpandHeadings.Add(TreeNode);
 
@@ -244,6 +283,18 @@ namespace ProjectLauncher
 		return *TreeNode.Get();
 	}
 
+	FLaunchProfileTreeNode& FLaunchProfileTreeNode::FindOrAddSubHeading(ILauncherProfileUATCommandRef InUATCommand)
+	{
+		if (SubHeadingNodes.Contains(InUATCommand->GetInternalName()))
+		{
+			return GetSubHeading(InUATCommand->GetInternalName());
+		}
+
+		FLaunchProfileTreeNode& TreeNode = AddSubHeading(InUATCommand->GetInternalName(), FText::FromString(*InUATCommand->GetDescription()), InUATCommand->GetOrder()); // @todo: need description and sort order to auto-update
+		TreeNode.UATCommand = InUATCommand;
+
+		return TreeNode;
+	}
 
 	FLaunchProfileTreeNode& FLaunchProfileTreeNode::AddWidget( FText InName, FCallbacks&& InWidgetCallbacks, TSharedRef<SWidget> InValueWidget )
 	{
@@ -253,6 +304,7 @@ namespace ProjectLauncher
 		TreeNode->Depth = Depth+1;
 		TreeNode->Widget = InValueWidget;
 		TreeNode->Callbacks = InWidgetCallbacks;
+		TreeNode->UATCommand = UATCommand;
 		Children.Add(TreeNode);
 
 		TreeData->RequestTreeRefresh();
diff --git a/Engine/Plugins/Developer/ProjectLauncher/Source/ProjectLauncher/Private/Styling/ProjectLauncherStyle.cpp b/Engine/Plugins/Developer/ProjectLauncher/Source/ProjectLauncher/Private/Styling/ProjectLauncherStyle.cpp
--- a/Engine/Plugins/Developer/ProjectLauncher/Source/ProjectLauncher/Private/Styling/ProjectLauncherStyle.cpp
+++ b/Engine/Plugins/Developer/ProjectLauncher/Source/ProjectLauncher/Private/Styling/ProjectLauncherStyle.cpp
@@ -74,6 +74,7 @@ void FProjectLauncherStyle::Initialize()
 	StyleSet->Set("Icons.DiffersFromDefault", new CORE_IMAGE_BRUSH_SVG( "../Editor/Slate/Starship/Common/ResetToDefault", Icon16x16) );
 	StyleSet->Set("Icons.DeviceManager", new CORE_IMAGE_BRUSH_SVG( "../Editor/Slate/Starship/Common/DeviceManager", Icon16x16) );
 	StyleSet->Set("Icons.Plugins", new CORE_IMAGE_BRUSH_SVG("../Editor/Slate/Starship/Common/Plugins", Icon16x16));
+	StyleSet->Set("Icons.Extensions", new CORE_IMAGE_BRUSH_SVG("../Editor/Slate/Starship/EditorViewport/menu", Icon16x16));
 
 	StyleSet->Set("Icons.Asset", new CORE_IMAGE_BRUSH( "../Editor/Slate/Icons/doc_16x", Icon16x16 ) );
 	StyleSet->Set("Icons.WarningWithColor.Small", new CORE_IMAGE_BRUSH_SVG("../Editor/Slate/Starship/Common/AlertTriangleSolid", Icon16x16, FStyleColors::Warning));
diff --git a/Engine/Plugins/Developer/ProjectLauncher/Source/ProjectLauncher/Private/Widgets/Output/SCustomLaunchLaunchPanel.cpp b/Engine/Plugins/Developer/ProjectLauncher/Source/ProjectLauncher/Private/Widgets/Output/SCustomLaunchLaunchPanel.cpp
--- a/Engine/Plugins/Developer/ProjectLauncher/Source/ProjectLauncher/Private/Widgets/Output/SCustomLaunchLaunchPanel.cpp
+++ b/Engine/Plugins/Developer/ProjectLauncher/Source/ProjectLauncher/Private/Widgets/Output/SCustomLaunchLaunchPanel.cpp
@@ -435,7 +435,7 @@ const FSlateBrush* SCustomLaunchLaunchPanel::GetProfileImage() const
 	if (LauncherProfile != nullptr)
 	{
 		// @todo: should there be a dedicated icon for multi-platform?
-		TArray<const PlatformInfo::FTargetPlatformInfo*> PlatformInfos = Model->GetPlatformInfos(LauncherProfile);
+		TArray<const PlatformInfo::FTargetPlatformInfo*> PlatformInfos = Model->GetPlatformInfosForAllCommands(LauncherProfile);
 		return FProjectLauncherStyle::GetProfileBrushForPlatforms(PlatformInfos, EPlatformIconSize::Large);
 	}
 
@@ -462,7 +462,7 @@ FText SCustomLaunchLaunchPanel::GetProfileDescription() const
 	{
 		if (Model->IsBasicLaunchProfile(LauncherProfile))
 		{
-			const TSharedPtr<ITargetDeviceProxy> DeviceProxy = Model->GetDeviceProxy(LauncherProfile.ToSharedRef());
+			const TSharedPtr<ITargetDeviceProxy> DeviceProxy = Model->GetDeviceProxy(Model->GetBasicLaunchProfile()->GetFirstBuildCookRun().ToSharedRef());
 			if (DeviceProxy.IsValid())
 			{
 				return FText::FromString(DeviceProxy->GetName());
@@ -527,9 +527,19 @@ FText SCustomLaunchLaunchPanel::GetProfileContentSchemeName() const
 		{
 			return LOCTEXT("AdvancedProfileInfo", "Advanced/Legacy Profile");
 		}
-
-		ProjectLauncher::EContentScheme ContentScheme = Model->DetermineProfileContentScheme(LauncherProfile.ToSharedRef());
-		return ProjectLauncher::GetContentSchemeDisplayName(ContentScheme);
+		else if (LauncherProfile->GetBuildCookRunCommands().Num() > 1)
+		{
+			return LOCTEXT("MultipleContentSchemeLabel", "Multiple"); // @todo: worth trying something fancy here? collecting all content schemes we're using
+		}
+		else if (!LauncherProfile->GetFirstBuildCookRun().IsValid())
+		{
+			return LOCTEXT("NoBuildCookRunsLabel", "Custom Profile");
+		}
+		else
+		{
+			ProjectLauncher::EContentScheme ContentScheme = Model->DetermineProfileContentScheme(LauncherProfile->GetFirstBuildCookRun().ToSharedRef());
+			return ProjectLauncher::GetContentSchemeDisplayName(ContentScheme);
+		}
 	}
 
 	return FText::GetEmpty();
diff --git a/Engine/Plugins/Developer/ProjectLauncher/Source/ProjectLauncher/Private/Widgets/Profiles/SCustomLaunchCustomProfileEditor.cpp b/Engine/Plugins/Developer/ProjectLauncher/Source/ProjectLauncher/Private/Widgets/Profiles/SCustomLaunchCustomProfileEditor.cpp
--- a/Engine/Plugins/Developer/ProjectLauncher/Source/ProjectLauncher/Private/Widgets/Profiles/SCustomLaunchCustomProfileEditor.cpp
+++ b/Engine/Plugins/Developer/ProjectLauncher/Source/ProjectLauncher/Private/Widgets/Profiles/SCustomLaunchCustomProfileEditor.cpp
@@ -6,6 +6,7 @@
 #include "Styling/StyleColors.h"
 #include "Styling/ProjectLauncherStyle.h"
 #include "Extension/LaunchExtension.h"
+#include "Extension/BuildCookRunLaunchExtension.h"
 #include "Framework/Commands/UIAction.h"
 #include "Framework/MultiBox/MultiBoxBuilder.h"
 #include "SlateOptMacros.h"
@@ -13,8 +14,10 @@
 #include "Widgets/Input/SHyperlink.h"
 #include "Widgets/Input/SComboBox.h"
 #include "Widgets/Input/SComboButton.h"
+#include "Widgets/Input/SCheckBox.h"
 #include "Widgets/Layout/SSplitter.h"
 #include "Widgets/Layout/SScrollBorder.h"
+#include "Widgets/Layout/SSpacer.h"
 #include "Widgets/Text/STextBlock.h"
 #include "Widgets/Text/SInlineEditableTextBlock.h"
 #include "Widgets/Images/SImage.h"
@@ -42,6 +45,40 @@ namespace ProjectLauncher
 	extern TSharedRef<ILaunchProfileTreeBuilder> CreateTreeBuilder( const ILauncherProfilePtr& InProfile, TSharedRef<ProjectLauncher::FModel> InModel );
 };
 
+class FProfileSectionDragDropOp : public FDragDropOperation
+{
+public:
+	DRAG_DROP_OPERATOR_TYPE(FProfileSectionDragDropOp, FDragDropOperation)
+
+	static TSharedRef<FProfileSectionDragDropOp> New(const ProjectLauncher::FLaunchProfileTreeNodePtr& InDraggedNode)
+	{
+		TSharedRef<FProfileSectionDragDropOp> Operation = MakeShared<FProfileSectionDragDropOp>();
+		Operation->DraggedNode = InDraggedNode;
+		Operation->Construct();
+		return Operation;
+	}
+
+	ProjectLauncher::FLaunchProfileTreeNodePtr GetDraggeedNode() const { return DraggedNode; }
+
+	virtual TSharedPtr<SWidget> GetDefaultDecorator() const override
+	{
+		return SNew(SBorder)
+			.BorderImage(FAppStyle::GetBrush("Brushes.Header"))
+			.Padding(8, 8, 200, 8)
+			.HAlign(HAlign_Left)
+			.VAlign(VAlign_Center)
+			[
+				SNew(STextBlock)
+				.Text(FText::FromString(DraggedNode->UATCommand->GetDescription()))
+				.Font(FCoreStyle::Get().GetFontStyle("SmallFontBold"))
+			];
+	}
+
+private:
+	ProjectLauncher::FLaunchProfileTreeNodePtr DraggedNode;
+};
+
+
 
 
 class SLaunchProfileCategoryTreeRow : public SLaunchProfileTreeRow
@@ -50,9 +87,17 @@ public:
 	SLATE_BEGIN_ARGS(SLaunchProfileCategoryTreeRow){}
 	SLATE_END_ARGS()
 
-	void Construct(const FArguments& InArgs, ProjectLauncher::FLaunchProfileTreeNodePtr TreeNode, const TSharedRef<STableViewBase>& OwnerTable )
+	ProjectLauncher::FLaunchProfileTreeNodePtr TreeNode;
+	bool bIsPrimaryHeading;
+	SCustomLaunchCustomProfileEditor* ProfileEditor;
+
+	void Construct(const FArguments& InArgs, ProjectLauncher::FLaunchProfileTreeNodePtr InTreeNode, SCustomLaunchCustomProfileEditor* InProfileEditor, bool bInIsPrimaryHeading )
 	{
-		auto HasVisibleChildren = [TreeNode]()
+		TreeNode = InTreeNode;
+		ProfileEditor = InProfileEditor;
+		bIsPrimaryHeading = bInIsPrimaryHeading;
+
+		auto HasVisibleChildren = [this]()
 		{
 			for (ProjectLauncher::FLaunchProfileTreeNodePtr ChildNode : TreeNode->Children)
 			{
@@ -65,13 +110,13 @@ public:
 			return false;
 		};
 
-		auto GetErrorVisibility = [this, TreeNode]()
+		auto GetErrorVisibility = [this]()
 		{
 			if (!IsItemExpanded())
 			{
 				for (ProjectLauncher::FLaunchProfileTreeNodePtr ChildNode : TreeNode->Children)
 				{
-					if (ChildNode->Callbacks.Validation.HasError( TreeNode->GetTreeData()->Profile.ToSharedRef() ))
+					if (ChildNode->Callbacks.Validation.HasError( TreeNode->GetTreeData()->Profile.ToSharedRef(), ChildNode->UATCommand ))
 					{
 						return EVisibility::Visible;
 					}
@@ -81,12 +126,12 @@ public:
 			return EVisibility::Collapsed;
 		};
 
-		auto GetErrorToolTipText = [TreeNode]()
+		auto GetErrorToolTipText = [this]()
 		{
 			TArray<FString> ErrorLines;
 			for (ProjectLauncher::FLaunchProfileTreeNodePtr ChildNode : TreeNode->Children)
 			{
-				ChildNode->Callbacks.Validation.GetErrorText( TreeNode->GetTreeData()->Profile.ToSharedRef(), ErrorLines );
+				ChildNode->Callbacks.Validation.GetErrorText( TreeNode->GetTreeData()->Profile.ToSharedRef(), ChildNode->UATCommand, ErrorLines );
 			}
 
 			FTextBuilder TextBuilder;
@@ -98,8 +143,42 @@ public:
 			return TextBuilder.ToText();
 		};
 
+		auto GetUATCommandEnabledCheckState = [this]()
+		{
+			if (TreeNode->UATCommand->IsEnabled())
+			{
+				return ECheckBoxState::Checked;
+			}
+			
+			return ECheckBoxState::Unchecked;
+		};
+
+		auto SetUATCommandEnabledCheckState = [this](ECheckBoxState State)
+		{
+			TreeNode->UATCommand->SetEnabled( State == ECheckBoxState::Checked);
+			ProfileEditor->TreeBuilder->OnPropertyChanged();
+		};
 
+		auto GetUATCommandDescription = [this]()
+		{
+			return FText::FromString(TreeNode->UATCommand->GetDescription());
+		};
+
+		auto SetUATCommandDescription = [this]( const FText& NewValue, ETextCommit::Type )
+		{ 
+			TreeNode->UATCommand->SetDescription(NewValue.ToString()); 
+			ProfileEditor->TreeBuilder->OnPropertyChanged();
+		};
 
+		auto OnRenameUATCommandClicked = [this]()
+		{
+			ProfileEditor->BeginRenameUATCommand(TreeNode->UATCommand.ToSharedRef());
+			return FReply::Handled();
+		};
+
+
+
+		TSharedPtr<SHorizontalBox> HeaderBar;
 		this->ChildSlot
 		[
 			SNew(SBorder)
@@ -114,8 +193,9 @@ public:
 					SNew(SBox)
 					.MinDesiredHeight(26.0f)
 					[
-						SNew(SHorizontalBox)
+						SAssignNew(HeaderBar, SHorizontalBox)
 
+						// tree expander arrow
 						+SHorizontalBox::Slot()
 						.HAlign(HAlign_Left)
 						.VAlign(VAlign_Center)
@@ -126,38 +206,165 @@ public:
 							.StyleSet( &FCoreStyle::Get() )
 						]
 
-						+SHorizontalBox::Slot()
-						.VAlign(VAlign_Center)
-						.Padding(4, 0, 0, 0)
-						.AutoWidth()
-						[
-							SNew(STextBlock)
-							.Text(TreeNode->Name)
-							.Font(FCoreStyle::Get().GetFontStyle("SmallFontBold"))
-						]
-
-						+SHorizontalBox::Slot()
-						.HAlign(HAlign_Left)
-						.VAlign(VAlign_Center)
-						.Padding(2, 0)
-						.AutoWidth()
-						[
-							SNew(SImage)
-							.Image(FProjectLauncherStyle::Get().GetBrush("Icons.WarningWithColor.Small"))
-							.Visibility_Lambda(GetErrorVisibility)
-							.ToolTipText_Lambda(GetErrorToolTipText)
-						]
 					]
 				]
 			]
 		];
 
+		if (bIsPrimaryHeading)
+		{
+			// checkbox to enable the UAT command
+			HeaderBar->AddSlot()
+			.HAlign(HAlign_Left)
+			.VAlign(VAlign_Center)
+			.Padding(4,0,0,0)
+			.AutoWidth()
+			[
+				SNew(SCheckBox)
+				.IsChecked_Lambda(GetUATCommandEnabledCheckState)
+				.OnCheckStateChanged_Lambda(SetUATCommandEnabledCheckState)
+			];
+		}
+
+		// error message indicator - only shown when the heading is collapsed to indicate a child node has an error
+		HeaderBar->AddSlot()
+		.HAlign(HAlign_Left)
+		.VAlign(VAlign_Center)
+		.Padding(2, 0)
+		.AutoWidth()
+		[
+			SNew(SImage)
+			.Image(FProjectLauncherStyle::Get().GetBrush("Icons.WarningWithColor.Small"))
+			.Visibility_Lambda(GetErrorVisibility)
+			.ToolTipText_Lambda(GetErrorToolTipText)
+		];
+
+
+		if (bIsPrimaryHeading)
+		{
+			// inline editable name field for UAT commands etc
+			TSharedPtr<SInlineEditableTextBlock> RenameBox;
+			HeaderBar->AddSlot()
+			.VAlign(VAlign_Center)
+			.Padding(4, 0, 0, 0)
+			.AutoWidth()
+			[
+				SAssignNew(RenameBox,SInlineEditableTextBlock)
+				.Text_Lambda(GetUATCommandDescription)
+				.OnTextCommitted_Lambda(SetUATCommandDescription)
+				.Font(FCoreStyle::Get().GetFontStyle("SmallFontBold"))
+			];
+			ProfileEditor->UATCommandNameFields.Add(TreeNode->UATCommand, RenameBox);
+		}
+		else
+		{
+			// basic display name
+			HeaderBar->AddSlot()
+			.VAlign(VAlign_Center)
+			.Padding(4, 0, 0, 0)
+			.AutoWidth()
+			[
+				SNew(STextBlock)
+				.Text(TreeNode->Name)
+				.Font(FCoreStyle::Get().GetFontStyle("SmallFontBold"))
+			];
+		}
+
+		if (bIsPrimaryHeading && TreeNode->UATCommand->AsBuildCookRun().IsValid())
+		{
+			ILauncherProfileBuildCookRunRef BuildCookRun = TreeNode->UATCommand->AsBuildCookRun().ToSharedRef();
+
+			// spacer to push everthing else to the right
+			HeaderBar->AddSlot()
+			.FillWidth(1)
+			[
+				SNew(SSpacer)
+			];
+
+			// drop-down menu
+			HeaderBar->AddSlot()
+			.HAlign(HAlign_Right)
+			.VAlign(VAlign_Center)
+			.Padding(2, 0, 16, 0)
+			.AutoWidth()
+			[
+				SNew(SComboButton)
+				.ButtonStyle(FAppStyle::Get(), "SimpleButton")
+				.ToolTipText(LOCTEXT("BCRExtensionMenuTip", "Enable an extension for this Build, Cook & Run"))
+				.OnGetMenuContent( ProfileEditor, &SCustomLaunchCustomProfileEditor::MakeBuildCookRunExtensionsSubMenuWidget, BuildCookRun)
+				.ButtonContent()
+				[
+					SNew(SImage)
+					.Image(FProjectLauncherStyle::Get().GetBrush("Icons.Plugins"))
+					.ColorAndOpacity(FSlateColor::UseSubduedForeground())
+				]
+			];
+		}
+
 		SLaunchProfileTreeRow::ConstructInternal(
 			STableRow::FArguments()
+			.OnDragDetected(this, &SLaunchProfileCategoryTreeRow::HandleDragDetected)		
+			.OnCanAcceptDrop(this, &SLaunchProfileCategoryTreeRow::HandleCanAcceptDrop)
+			.OnAcceptDrop(this, &SLaunchProfileCategoryTreeRow::HandleAcceptDrop)
 			.ShowSelection(false),
-			OwnerTable
+			ProfileEditor->GetTreeView()
 		);
 	}
+
+	virtual FReply OnMouseButtonUp(const FGeometry& Geometry, const FPointerEvent& MouseEvent) override
+	{
+		if (MouseEvent.GetEffectingButton() == EKeys::RightMouseButton && bIsPrimaryHeading)
+		{
+			FMenuBuilder MenuBuilder(true, nullptr);
+			ProfileEditor->MakeUATCommandSubMenu(MenuBuilder, TreeNode->UATCommand.ToSharedRef());
+
+			FSlateApplication::Get().PushMenu( SharedThis(this),
+				FWidgetPath(),
+				MenuBuilder.MakeWidget(),
+				MouseEvent.GetScreenSpacePosition(),
+				FPopupTransitionEffect(FPopupTransitionEffect::ContextMenu));
+
+			return FReply::Handled();
+		}
+
+		return SLaunchProfileTreeRow::OnMouseButtonUp(Geometry, MouseEvent);
+	}
+
+	FReply HandleDragDetected(const FGeometry& Geometry, const FPointerEvent& MouseEvent)
+	{
+		if (bIsPrimaryHeading && MouseEvent.IsMouseButtonDown(EKeys::LeftMouseButton))
+		{
+			TSharedPtr<FDragDropOperation> DragDropOp = FProfileSectionDragDropOp::New(TreeNode);
+			return FReply::Handled().BeginDragDrop(DragDropOp.ToSharedRef());
+		}
+
+		return FReply::Unhandled();
+	}
+
+	TOptional<EItemDropZone> HandleCanAcceptDrop(const FDragDropEvent& InDragDropEvent, EItemDropZone InItemDropZone, ProjectLauncher::FLaunchProfileTreeNodePtr InTreeItem)
+	{
+		TSharedPtr<FProfileSectionDragDropOp> DragDropOp = InDragDropEvent.GetOperationAs<FProfileSectionDragDropOp>();
+		if (DragDropOp.IsValid() && bIsPrimaryHeading && InItemDropZone != EItemDropZone::OntoItem)
+		{
+			return TOptional<EItemDropZone>(InItemDropZone);
+		}
+
+		return TOptional<EItemDropZone>();
+	}
+
+
+	FReply HandleAcceptDrop(const FDragDropEvent& InDragDropEvent, EItemDropZone InItemDropZone, ProjectLauncher::FLaunchProfileTreeNodePtr InTreeItem)
+	{
+		TSharedPtr<FProfileSectionDragDropOp> DragDropOp = InDragDropEvent.GetOperationAs<FProfileSectionDragDropOp>();
+		if (DragDropOp.IsValid())
+		{
+			ProfileEditor->ReorderNodes(DragDropOp->GetDraggeedNode(), TreeNode, InItemDropZone);
+			return FReply::Handled();
+		}
+
+		return FReply::Unhandled();
+	}
+
 };
 
 class SLaunchProfileFullWidthTreeRow : public SLaunchProfileTreeRow
@@ -166,7 +373,7 @@ public:
 	SLATE_BEGIN_ARGS(SLaunchProfileFullWidthTreeRow){}
 	SLATE_END_ARGS()
 
-	void Construct(const FArguments& InArgs, ProjectLauncher::FLaunchProfileTreeNodePtr TreeNode, const TSharedRef<STableViewBase>& OwnerTable )
+	void Construct(const FArguments& InArgs, ProjectLauncher::FLaunchProfileTreeNodePtr TreeNode, SCustomLaunchCustomProfileEditor* ProfileEditor )
 	{
 		this->ChildSlot
 		[
@@ -191,7 +398,7 @@ public:
 		SLaunchProfileTreeRow::ConstructInternal(
 			STableRow::FArguments()
 			.ShowSelection(false),
-			OwnerTable
+			ProfileEditor->GetTreeView()
 		);
 	}
 };
@@ -209,7 +416,7 @@ public:
 	SLATE_END_ARGS()
 
 	BEGIN_SLATE_FUNCTION_BUILD_OPTIMIZATION
-	void Construct(const FArguments& InArgs, ProjectLauncher::FLaunchProfileTreeNodePtr TreeNode, const TSharedRef<STableViewBase>& OwnerTable )
+	void Construct(const FArguments& InArgs, ProjectLauncher::FLaunchProfileTreeNodePtr TreeNode, SCustomLaunchCustomProfileEditor* ProfileEditor )
 	{
 		SplitterValue = InArgs._SplitterValue;
 		OnSplitterResized = InArgs._OnSplitterResized;
@@ -257,7 +464,7 @@ public:
 
 		auto GetErrorVisibility = [TreeNode]()
 		{
-			if (TreeNode->Callbacks.Validation.HasError( TreeNode->GetTreeData()->Profile.ToSharedRef() ))
+			if (TreeNode->Callbacks.Validation.HasError( TreeNode->GetTreeData()->Profile.ToSharedRef(), TreeNode->UATCommand ))
 			{
 				return EVisibility::Visible;
 			}
@@ -268,7 +475,7 @@ public:
 		auto GetErrorToolTipText = [TreeNode]()
 		{
 			TArray<FString> ErrorLines;
-			TreeNode->Callbacks.Validation.GetErrorText( TreeNode->GetTreeData()->Profile.ToSharedRef(), ErrorLines );
+			TreeNode->Callbacks.Validation.GetErrorText( TreeNode->GetTreeData()->Profile.ToSharedRef(), TreeNode->UATCommand, ErrorLines );
 
 			FTextBuilder TextBuilder;
 			for (const FString& ErrorLine : ErrorLines)
@@ -391,7 +598,7 @@ public:
 		SLaunchProfileTreeRow::ConstructInternal(
 			STableRow::FArguments()
 			.ShowSelection(false),
-			OwnerTable
+			ProfileEditor->GetTreeView()
 		);
 	}
 	END_SLATE_FUNCTION_BUILD_OPTIMIZATION
@@ -477,7 +684,7 @@ void SCustomLaunchCustomProfileEditor::Construct(const FArguments& InArgs, const
 			[
 				SAssignNew(TreeView, SLaunchProfileTreeView)
 				.TreeItemsSource(&TreeBuilder->GetProfileTree()->Nodes)
-				.SelectionMode( ESelectionMode::None )
+				.SelectionMode( ESelectionMode::Single )
 				.OnGenerateRow( this, &SCustomLaunchCustomProfileEditor::OnGenerateWidgetForTreeNode )
 				.OnGetChildren( this, &SCustomLaunchCustomProfileEditor::OnGetChildren)
 				.HandleDirectionalNavigation(false)
@@ -502,14 +709,14 @@ TSharedRef<ITableRow> SCustomLaunchCustomProfileEditor::OnGenerateWidgetForTreeN
 
 	if (TreeNode->Name.IsEmpty() && TreeNode->Widget.IsValid())
 	{
-		return SNew(SLaunchProfileFullWidthTreeRow, TreeNode, OwnerTable)
+		return SNew(SLaunchProfileFullWidthTreeRow, TreeNode, this)
 		.Visibility_Lambda( [IsVisible]() { return IsVisible() ? EVisibility::Visible : EVisibility::Collapsed; } )
 		.IsEnabled_Lambda( [TreeNode]() { return TreeNode->Callbacks.IsEnabled ? TreeNode->Callbacks.IsEnabled() : true; } )
 		;
 	}
 	else if (TreeNode->Widget.IsValid())
 	{
-		return SNew(SLaunchProfilePropertyTreeRow, TreeNode, OwnerTable)
+		return SNew(SLaunchProfilePropertyTreeRow, TreeNode, this)
 		.SplitterValue_Lambda( [this]() { return SplitterPos; } )
 		.OnSplitterResized_Lambda( [this](float NewPos) { SplitterPos = FMath::Clamp(NewPos, 0.0f, 1.0f); } )
 		.Visibility_Lambda( [IsVisible]() { return IsVisible() ? EVisibility::Visible : EVisibility::Collapsed; } )
@@ -517,7 +724,68 @@ TSharedRef<ITableRow> SCustomLaunchCustomProfileEditor::OnGenerateWidgetForTreeN
 	}
 	else
 	{
-		return SNew(SLaunchProfileCategoryTreeRow, TreeNode, OwnerTable);
+		bool bIsPrimaryHeading = TreeNode->UATCommand.IsValid() && TreeBuilder->GetProfileTree()->UATCommandHeadings.Contains(TreeNode);
+
+		return SNew(SLaunchProfileCategoryTreeRow, TreeNode, this, bIsPrimaryHeading);
+	}
+}
+
+void SCustomLaunchCustomProfileEditor::ReorderNodes(ProjectLauncher::FLaunchProfileTreeNodePtr SourceNode, ProjectLauncher::FLaunchProfileTreeNodePtr TargetNode, EItemDropZone DropZone)
+{
+	// do nothing if we're dropped on ourselves
+	if (SourceNode == TargetNode)
+	{
+		return;
+	}
+
+	TArray<ProjectLauncher::FLaunchProfileTreeNodePtr>& Nodes = TreeBuilder->GetProfileTree()->Nodes;
+	if (!Nodes.Contains(SourceNode) || !Nodes.Contains(TargetNode))
+	{
+		return;
+	}
+
+	// move the item to where we want it
+	Nodes.Remove(SourceNode);
+	int32 TargetIndex = Nodes.IndexOfByKey(TargetNode);
+	if (DropZone != EItemDropZone::AboveItem)
+	{
+		TargetIndex++;
+	}
+	Nodes.Insert(SourceNode, TargetIndex);
+
+	// ensure the order of items is contiguous
+	// @todo: maybe do away with the magic numbers for initial ordering?
+	int32 PrevSortOrder = INT32_MIN;
+	for ( const ProjectLauncher::FLaunchProfileTreeNodePtr& Node : Nodes)
+	{
+		if (Node->SortOrder <= PrevSortOrder && Node->SortOrder != INT32_MIN)
+		{
+			Node->SortOrder = PrevSortOrder+1;
+			if (Node->UATCommand.IsValid())
+			{
+				Node->UATCommand->SetOrder(Node->SortOrder);
+			}
+		}
+
+		PrevSortOrder = Node->SortOrder;
+	}
+
+	TreeBuilder->OnPropertyChanged();
+	TreeView->RequestTreeRefresh();
+}
+
+void SCustomLaunchCustomProfileEditor::BeginRenameUATCommand( const ILauncherProfileUATCommandRef& UATCommand ) const
+{
+	const TWeakPtr<SInlineEditableTextBlock>* RenameBoxPtr = UATCommandNameFields.Find(UATCommand.ToSharedPtr());
+	if (RenameBoxPtr != nullptr)
+	{
+		ExecuteOnGameThread( UE_SOURCE_LOCATION, [RenameBox = (*RenameBoxPtr)]
+		{
+			if (RenameBox.IsValid())
+			{
+				RenameBox.Pin()->EnterEditingMode(); 
+			}
+		});
 	}
 }
 
@@ -554,18 +822,36 @@ TSharedRef<SWidget> SCustomLaunchCustomProfileEditor::GetScrollWidget()
 
 void SCustomLaunchCustomProfileEditor::SetProfile( const ILauncherProfilePtr& Profile )
 {
-	CurrentProfile = Profile;
-
+	// save the current node collapsed state (nb. not saved to the profile)
+	if (CurrentProfile.IsValid() && TreeBuilder.IsValid())
+	{
+		TSet<FString>& CollapsedState = CachedCollapsedNodes.Add(CurrentProfile);
+		for ( TTuple<FString,ProjectLauncher::FLaunchProfileTreeNodePtr> Pair : TreeBuilder->GetProfileTree()->HeadingNodes)
+		{
+			CacheCollapsedRecursive(CollapsedState, Pair.Key, Pair.Value);
+		}
+	}
 
+	// change the profile and rebuild
+	CurrentProfile = Profile;
 	TreeBuilder = ProjectLauncher::CreateTreeBuilder( Profile, Model.ToSharedRef() );
-
 	TreeView->SetTreeItemsSource( &TreeBuilder->GetProfileTree()->Nodes );
-	for ( ProjectLauncher::FLaunchProfileTreeNodePtr Node : TreeView->GetRootItems())
+	TreeView->RequestTreeRefresh();
+
+	// restore the collapsed state if we have one saved, otherwise expand everything (need to defer until after the tree is rebuilt)
+	ExecuteOnGameThread( UE_SOURCE_LOCATION, [this]
 	{
-		TreeView->SetItemExpansion(Node, true);
-	}
+		TSet<FString> CollapsedState;
+		if (CachedCollapsedNodes.Contains(CurrentProfile))
+		{
+			CollapsedState = CachedCollapsedNodes.FindAndRemoveChecked(CurrentProfile);
+		}
+		for ( TTuple<FString,ProjectLauncher::FLaunchProfileTreeNodePtr> Pair : TreeBuilder->GetProfileTree()->HeadingNodes)
+		{
+			ApplyCollapsedRecusive(CollapsedState, Pair.Key, Pair.Value);
+		}
+	});
 
-	TreeView->RequestTreeRefresh();
 }
 
 void SCustomLaunchCustomProfileEditor::RebuildTree()
@@ -623,6 +909,10 @@ void SCustomLaunchCustomProfileEditor::ApplyCollapsedRecusive( TSet<FString>& Co
 			ApplyCollapsedRecusive(Collapsed,  NodeName + TEXT(":") + Pair.Key, Pair.Value);
 		}
 	}
+	else
+	{
+		TreeView->SetItemExpansion(Node, false);
+	}
 }
 
 
@@ -658,6 +948,13 @@ TSharedRef<SWidget> SCustomLaunchCustomProfileEditor::MakeExtensionsMenu()
 			continue;
 		}
 
+		// don't show additional UAT commands or automated tests in the Extensions menu - save them for the Tasks menu
+		if (Extension->AsUATCommandLaunchExtension() || Extension->AsAutomatedTestLaunchExtension())
+		{
+			continue;
+		}
+
+
 		// get section
 		FString SectionName = ExtensionMenuEntry.SectionName;
 		FText SectionDisplayName = ExtensionMenuEntry.SectionDisplayName;
@@ -691,6 +988,7 @@ TSharedRef<SWidget> SCustomLaunchCustomProfileEditor::MakeExtensionsMenu()
 		}
 	}
 
+
 	// sort sections alphabetically
 	Sections.ValueSort([](const FExtensionMenuSection& A, const FExtensionMenuSection& B)
 	{
@@ -707,80 +1005,7 @@ TSharedRef<SWidget> SCustomLaunchCustomProfileEditor::MakeExtensionsMenu()
 	FMenuBuilder MenuBuilder(bShouldCloseWindowAfterMenuSelection, nullptr, nullptr, bCloseSelfOnly, &FCoreStyle::Get(), bSearchable );
 	
 
-	auto MakeExtensionMenu = [this](FMenuBuilder& MenuBuilder, TSharedRef<ProjectLauncher::FLaunchExtension> Extension)
-	{
-		TSharedPtr<ProjectLauncher::FLaunchExtensionInstance> ExtensionInstance = ProjectLauncher::FLaunchExtension::GetExtensionInstance(Extension, CurrentProfile.ToSharedRef());
-
-		if (!Extension->IsAlwaysCreated(CurrentProfile.ToSharedRef(), Model.ToSharedRef()))
-		{
-			auto OnToggleExtension = [this, Extension, ExtensionInstance]()
-			{
-				if (ExtensionInstance.IsValid())
-				{
-					// defer the delete until the next tick
-					PendingExtensionDeletes.AddUnique(Extension);
-				}
-				else
-				{
-					TSharedPtr<ProjectLauncher::FLaunchExtensionInstance> Instance = ProjectLauncher::FLaunchExtension::CreateExtensionInstance(Extension, CurrentProfile.ToSharedRef(), Model.ToSharedRef(), TreeBuilder->GetProfileTree().Get());
-					if (Instance.IsValid())
-					{
-						Instance->CustomizeTree(TreeBuilder->GetProfileTree().Get());
-						TreeBuilder->OnPropertyChanged();
-						CurrentProfile->RefreshCustomWarningsAndErrors();
-					}
-				}
-				
-			};
-
-			auto CanToggleExtension = [this, Extension, ExtensionInstance]
-			{
-				if (ExtensionInstance.IsValid())
-				{
-					return ExtensionInstance->CanBeRemoved();
-				}
-				else
-				{
-					return Extension->CanBeCreated(CurrentProfile.ToSharedRef(), Model.ToSharedRef());
-				}
-			};
-
-			auto GetToolTipText = [this, Extension, ExtensionInstance]
-			{
-				if (ExtensionInstance.IsValid() && !ExtensionInstance->CanBeRemoved())
-				{
-					return LOCTEXT("CannotRemoveExtensionTooltip", "The extension cannot be removed at this time");
-				}
-				else if (!ExtensionInstance.IsValid() && !Extension->CanBeCreated(CurrentProfile.ToSharedRef(), Model.ToSharedRef()))
-				{
-					return LOCTEXT("CannotAddExtensionTooltip", "The extension cannot be added at this time");
-				}
-
-				return FText::GetEmpty();
-			};
-
-			MenuBuilder.AddMenuEntry(
-				Extension->GetDisplayName(),
-				GetToolTipText(),
-				FSlateIcon(),
-				FUIAction(
-					FExecuteAction::CreateLambda(OnToggleExtension),
-					FCanExecuteAction::CreateLambda(CanToggleExtension),
-					FIsActionChecked::CreateLambda( [ExtensionInstance]() { return ExtensionInstance.IsValid(); } )
-				),
-				NAME_None,
-				EUserInterfaceActionType::Check);
-		}
-
-		if (ExtensionInstance)
-		{
-			ExtensionInstance->MakeCustomExtensionSubmenu(MenuBuilder);
-		}
-
-		Extension->MakeCustomExtensionSubmenu(MenuBuilder, CurrentProfile.ToSharedRef(), Model.ToSharedRef());
-	};
-
-
+	bool bIncludeBuildCookRunExtensions = Model->CanUseSimplifiedLayout(CurrentProfile.ToSharedRef());
 	for (const auto& SectionItr : Sections)
 	{
 		// start the section
@@ -789,18 +1014,18 @@ TSharedRef<SWidget> SCustomLaunchCustomProfileEditor::MakeExtensionsMenu()
 		// add all extensions
 		for (TSharedRef<ProjectLauncher::FLaunchExtension> Extension : SectionItr.Value.Extensions)
 		{
-			MakeExtensionMenu(MenuBuilder, Extension);
+			MakeExtensionMenu(MenuBuilder, Extension, bIncludeBuildCookRunExtensions);
 		}
 
 		// add all submenus
 		for (const auto& SubMenuItr : SectionItr.Value.SubMenus)
 		{
-			auto MakeExtensionSubMenu = [MakeExtensionMenu](FMenuBuilder& MenuBuilder, TSharedPtr<FExtensionMenuSubMenu> SubMenu)
+			auto MakeExtensionSubMenu = [this, bIncludeBuildCookRunExtensions](FMenuBuilder& MenuBuilder, TSharedPtr<FExtensionMenuSubMenu> SubMenu)
 			{
 				MenuBuilder.SetSearchable(bSearchable);
 				for (TSharedRef<ProjectLauncher::FLaunchExtension> Extension : SubMenu->Extensions)
 				{
-					MakeExtensionMenu(MenuBuilder, Extension);
+					MakeExtensionMenu(MenuBuilder, Extension, bIncludeBuildCookRunExtensions);
 				}
 			};
 
@@ -822,6 +1047,256 @@ TSharedRef<SWidget> SCustomLaunchCustomProfileEditor::MakeExtensionsMenu()
 	return MenuBuilder.MakeWidget();
 }
 
+void SCustomLaunchCustomProfileEditor::MakeExtensionMenu( FMenuBuilder& MenuBuilder, TSharedRef<ProjectLauncher::FLaunchExtension> Extension, bool bIncludeBuildCookRunExtensions )
+{
+	TSharedPtr<ProjectLauncher::FLaunchExtensionInstance> ExtensionInstance = ProjectLauncher::FLaunchExtension::GetExtensionInstance(Extension, CurrentProfile.ToSharedRef());
+
+	if (!Extension->IsAlwaysCreated(CurrentProfile.ToSharedRef(), Model.ToSharedRef()))
+	{
+		auto OnToggleExtension = [this, Extension, ExtensionInstance]()
+		{
+			if (ExtensionInstance.IsValid())
+			{
+				// defer the delete until the next tick
+				PendingExtensionDeletes.AddUnique(Extension);
+			}
+			else
+			{
+				TSharedPtr<ProjectLauncher::FLaunchExtensionInstance> Instance = ProjectLauncher::FLaunchExtension::CreateExtensionInstance(Extension, CurrentProfile.ToSharedRef(), Model.ToSharedRef(), TreeBuilder->GetProfileTree().Get());
+				if (Instance.IsValid())
+				{
+					TreeBuilder->GetProfileTree()->RequestFullTreeRebuild();
+
+					TreeBuilder->OnPropertyChanged();
+					CurrentProfile->RefreshCustomWarningsAndErrors();
+				}
+			}
+				
+		};
+
+		auto CanToggleExtension = [this, Extension, ExtensionInstance]
+		{
+			if (ExtensionInstance.IsValid())
+			{
+				return ExtensionInstance->CanBeRemoved();
+			}
+			else
+			{
+				return Extension->CanBeCreated(CurrentProfile.ToSharedRef(), Model.ToSharedRef());
+			}
+		};
+
+		auto GetToolTipText = [this, Extension, ExtensionInstance]
+		{
+			if (ExtensionInstance.IsValid() && !ExtensionInstance->CanBeRemoved())
+			{
+				return LOCTEXT("CannotRemoveExtensionTooltip", "The extension cannot be removed at this time");
+			}
+			else if (!ExtensionInstance.IsValid() && !Extension->CanBeCreated(CurrentProfile.ToSharedRef(), Model.ToSharedRef()))
+			{
+				return LOCTEXT("CannotAddExtensionTooltip", "The extension cannot be added at this time");
+			}
+
+			return FText::GetEmpty();
+		};
+
+		MenuBuilder.AddMenuEntry(
+			Extension->GetDisplayName(),
+			GetToolTipText(),
+			FSlateIcon(),
+			FUIAction(
+				FExecuteAction::CreateLambda(OnToggleExtension),
+				FCanExecuteAction::CreateLambda(CanToggleExtension),
+				FIsActionChecked::CreateLambda( [ExtensionInstance]() { return ExtensionInstance.IsValid(); } )
+			),
+			NAME_None,
+			EUserInterfaceActionType::Check);
+	}
+
+	if (ExtensionInstance)
+	{
+		if (bIncludeBuildCookRunExtensions && CurrentProfile->GetFirstBuildCookRun().IsValid())
+		{
+			if (ProjectLauncher::IBuildCookRunExtensionFactory* BuildCookRunFactory = ExtensionInstance->AsBuildCookRunFactory())
+			{
+				BuildCookRunFactory->CustomizeBuildCookRunExtensionMenu(MenuBuilder, CurrentProfile->GetFirstBuildCookRun().ToSharedRef());
+			}
+		}
+
+		ExtensionInstance->MakeCustomExtensionSubmenu(MenuBuilder);
+	}
+
+	Extension->MakeCustomExtensionSubmenu(MenuBuilder, CurrentProfile.ToSharedRef(), Model.ToSharedRef());
+};
+
+
+
+
+
+
+TSharedRef<SWidget> SCustomLaunchCustomProfileEditor::MakeCommandsMenu()
+{
+	ProjectLauncher::FLaunchProfileTreeDataRef TreeData = TreeBuilder->GetProfileTree();
+
+	const bool bShouldCloseWindowAfterMenuSelection = true;
+	const bool bCloseSelfOnly = true;
+	const bool bSearchable = false;
+	const bool bInOpenSubMenuOnClick = false;
+	FMenuBuilder MenuBuilder(bShouldCloseWindowAfterMenuSelection, nullptr, nullptr, bCloseSelfOnly, &FCoreStyle::Get(), bSearchable );
+
+	// cannot add new BuildCookRuns to the basic launch profile
+	if (!Model->IsBasicLaunchProfile(CurrentProfile))
+	{
+		MenuBuilder.AddMenuEntry(
+			LOCTEXT("NewBuildCookRunCmdLabel", "Add New Build, Cook & Run"),
+			FText::GetEmpty(),
+			FSlateIcon(),
+			FUIAction(
+				FExecuteAction::CreateLambda( [this, TreeData]()
+				{
+					ILauncherProfileBuildCookRunRef NewBuildCookRun = TreeData->Model->CreateBuildCookRun(CurrentProfile.ToSharedRef());
+					NewBuildCookRun->SetDescription(TEXT("New Task (Build, Cook & Run)"));
+					
+					TreeData->RequestFullTreeRebuild();
+				})
+			),
+			NAME_None);
+	}
+
+	// enumerate extensions & automated tests & add them here too. 
+	// @todo: honor FExtensionsMenuEntry structure here
+	// @todo: UAT commands and AutomatedTests should also be add/remove like BuildCookRun
+	MenuBuilder.BeginSection(NAME_None, LOCTEXT("UATCommandsGroup", "Add Build Tasks"));
+	for (const TSharedPtr<ProjectLauncher::FLaunchExtension> Extension : ProjectLauncher::FLaunchExtension::GetExtensions())
+	{
+		if (Extension->AsUATCommandLaunchExtension() != nullptr)
+		{
+			MakeExtensionMenu(MenuBuilder, Extension.ToSharedRef(), false);
+		}
+	}
+	MenuBuilder.EndSection();
+
+	MenuBuilder.BeginSection(NAME_None, LOCTEXT("AutomatedTestsGroup", "Add Automated Tests"));
+	for (const TSharedPtr<ProjectLauncher::FLaunchExtension> Extension : ProjectLauncher::FLaunchExtension::GetExtensions())
+	{
+		if (Extension->AsAutomatedTestLaunchExtension() != nullptr)
+		{
+			MakeExtensionMenu(MenuBuilder, Extension.ToSharedRef(), false);
+		}
+	}
+	MenuBuilder.EndSection();
+
+	return MenuBuilder.MakeWidget();
+}
+
+void SCustomLaunchCustomProfileEditor::MakeUATCommandSubMenu( FMenuBuilder& MenuBuilder, ILauncherProfileUATCommandRef UATCommand )
+{
+	ProjectLauncher::FLaunchProfileTreeDataRef TreeData = TreeBuilder->GetProfileTree();
+
+	const bool bShouldCloseWindowAfterMenuSelection = true;
+	const bool bSearchable = false;
+	const bool bInOpenSubMenuOnClick = false;
+
+	MenuBuilder.SetSearchable(bSearchable);
+	MenuBuilder.AddMenuEntry(
+		LOCTEXT("EnableCommandLabel", "Enabled"),
+		FText::GetEmpty(),
+		FSlateIcon(),
+		FUIAction(
+			FExecuteAction::CreateLambda( [UATCommand]() { UATCommand->SetEnabled( !UATCommand->IsEnabled() ); } ),
+			FCanExecuteAction(),
+			FIsActionChecked::CreateLambda( [UATCommand]() { return UATCommand->IsEnabled(); } )
+		),
+		NAME_None,
+		EUserInterfaceActionType::Check);
+
+	if (UATCommandNameFields.Contains(UATCommand.ToSharedPtr()))
+	{
+		MenuBuilder.AddMenuEntry(
+			LOCTEXT("RenameCommandLabel", "Rename"),
+			FText::GetEmpty(),
+			FSlateIcon(),
+			FUIAction(
+				FExecuteAction::CreateLambda( [this, UATCommand]()
+				{
+					BeginRenameUATCommand(UATCommand);
+				})
+			),
+			NAME_None);
+	}
+
+	MenuBuilder.AddMenuEntry(
+		LOCTEXT("RemoveCommandLabel", "Delete"),
+		FText::GetEmpty(),
+		FSlateIcon(),
+		FUIAction(
+			FExecuteAction::CreateLambda( [this, TreeData, UATCommand]() 
+			{
+				// fixme: make this better
+				if (FPlatformMisc::MessageBoxExt(EAppMsgType::YesNo, TEXT("Confirm task deletion?"), TEXT("Delete task")) == EAppReturnType::Yes)
+				{
+					// see if there is an extension associated with this UAT command
+					TSharedPtr<ProjectLauncher::FLaunchExtensionInstance>* ExtensionInstancePtr = TreeData->ExtensionInstances.FindByPredicate( [UATCommand]( TSharedPtr<ProjectLauncher::FLaunchExtensionInstance> ExtensionInstance )
+					{
+						return ExtensionInstance->ManagesUATCommand(UATCommand);
+					});
+					TSharedPtr<ProjectLauncher::FLaunchExtension> Extension = (ExtensionInstancePtr != nullptr) ? (*ExtensionInstancePtr)->GetExtension().ToSharedPtr() : nullptr;
+
+
+					TreeData->Profile->RemoveUATCommand(UATCommand->GetInternalName()); 
+
+					if (Extension.IsValid())
+					{
+						// defer the delete until the next tick
+						PendingExtensionDeletes.AddUnique(Extension.ToSharedRef());
+					}
+					else
+					{
+						// likely a BuildCookRun
+						TreeBuilder->OnPropertyChanged();
+						TreeData->RequestFullTreeRebuild();
+					}
+				}
+			})
+		),
+		NAME_None);
+
+	ILauncherProfileBuildCookRunPtr BuildCookRun = UATCommand->AsBuildCookRun();
+	if (BuildCookRun.IsValid())
+	{
+		MenuBuilder.BeginSection(NAME_None, LOCTEXT("ExtensionsGroupLabel", "Extensions"));
+		MakeBuildCookRunExtensionsSubMenu(MenuBuilder, BuildCookRun.ToSharedRef());
+		MenuBuilder.EndSection();
+	}
+}
+
+void SCustomLaunchCustomProfileEditor::MakeBuildCookRunExtensionsSubMenu( FMenuBuilder& MenuBuilder, ILauncherProfileBuildCookRunRef BuildCookRun ) const
+{
+	ProjectLauncher::FLaunchProfileTreeDataRef TreeData = TreeBuilder->GetProfileTree();
+
+	for (TSharedPtr<ProjectLauncher::FLaunchExtensionInstance> ExtensionInstance : TreeData->ExtensionInstances)
+	{
+		if (ProjectLauncher::IBuildCookRunExtensionFactory* BuildCookRunFactory = ExtensionInstance->AsBuildCookRunFactory())
+		{
+			BuildCookRunFactory->CustomizeBuildCookRunExtensionMenu(MenuBuilder, BuildCookRun);
+		}
+	}
+}
+
+
+TSharedRef<SWidget> SCustomLaunchCustomProfileEditor::MakeBuildCookRunExtensionsSubMenuWidget(ILauncherProfileBuildCookRunRef BuildCookRun) const
+{
+	const bool bShouldCloseWindowAfterMenuSelection = true;
+	const bool bCloseSelfOnly = true;
+	const bool bSearchable = false;
+	const bool bInOpenSubMenuOnClick = false;
+	FMenuBuilder MenuBuilder(bShouldCloseWindowAfterMenuSelection, nullptr, nullptr, bCloseSelfOnly, &FCoreStyle::Get(), bSearchable );
+
+	MakeBuildCookRunExtensionsSubMenu(MenuBuilder, BuildCookRun);
+
+	return MenuBuilder.MakeWidget();
+}
+
 
 void SCustomLaunchCustomProfileEditor::Tick(const FGeometry& AllottedGeometry, const double InCurrentTime, const float InDeltaTime)
 {
diff --git a/Engine/Plugins/Developer/ProjectLauncher/Source/ProjectLauncher/Private/Widgets/Profiles/SCustomLaunchCustomProfileEditor.h b/Engine/Plugins/Developer/ProjectLauncher/Source/ProjectLauncher/Private/Widgets/Profiles/SCustomLaunchCustomProfileEditor.h
--- a/Engine/Plugins/Developer/ProjectLauncher/Source/ProjectLauncher/Private/Widgets/Profiles/SCustomLaunchCustomProfileEditor.h
+++ b/Engine/Plugins/Developer/ProjectLauncher/Source/ProjectLauncher/Private/Widgets/Profiles/SCustomLaunchCustomProfileEditor.h
@@ -38,8 +38,10 @@ public:
 	ILauncherProfilePtr GetProfile() const { return CurrentProfile; }
 
 	TSharedRef<SWidget> MakeExtensionsMenu();
+	TSharedRef<SWidget> MakeCommandsMenu();
 
 	void RebuildTree();
+	TSharedRef<SLaunchProfileTreeView> GetTreeView() const { return TreeView.ToSharedRef(); }
 
 	// IScrollableWidget interface
 	virtual FVector2D GetScrollDistance() override;
@@ -54,6 +56,12 @@ private:
 
 	void CacheCollapsedRecursive( TSet<FString>& Collapsed, const FString& NodeName, ProjectLauncher::FLaunchProfileTreeNodePtr Node) const;
 	void ApplyCollapsedRecusive( TSet<FString>& Collapsed, const FString& NodeName, ProjectLauncher::FLaunchProfileTreeNodePtr Node) const;
+	void ReorderNodes(ProjectLauncher::FLaunchProfileTreeNodePtr SourceNode, ProjectLauncher::FLaunchProfileTreeNodePtr TargetNode, EItemDropZone DropZone);
+
+	void MakeUATCommandSubMenu( FMenuBuilder& MenuBuilder, ILauncherProfileUATCommandRef UATCommand );
+	void MakeBuildCookRunExtensionsSubMenu( FMenuBuilder& MenuBuilder, ILauncherProfileBuildCookRunRef BuildCookRun ) const;
+	void MakeExtensionMenu( FMenuBuilder& MenuBuilder, TSharedRef<ProjectLauncher::FLaunchExtension> Extension, bool bIncludeBuildCookRunExtensions );
+	TSharedRef<SWidget> MakeBuildCookRunExtensionsSubMenuWidget(ILauncherProfileBuildCookRunRef BuildCookRun) const;
 
 	ILauncherProfilePtr CurrentProfile;
 	TSharedPtr<ProjectLauncher::FModel> Model;
@@ -64,6 +72,10 @@ private:
 	float SplitterPos = 0.6f;
 	
 
-private:
+	void BeginRenameUATCommand( const ILauncherProfileUATCommandRef& UATCommand ) const;
 
+private:
+	TMap<ILauncherProfilePtr,TSet<FString>> CachedCollapsedNodes;
+	TMap<ILauncherProfileUATCommandPtr,TWeakPtr<class SInlineEditableTextBlock>> UATCommandNameFields;
+	friend class SLaunchProfileCategoryTreeRow;
 };
diff --git a/Engine/Plugins/Developer/ProjectLauncher/Source/ProjectLauncher/Private/Widgets/Profiles/SCustomLaunchCustomProfileSelector.cpp b/Engine/Plugins/Developer/ProjectLauncher/Source/ProjectLauncher/Private/Widgets/Profiles/SCustomLaunchCustomProfileSelector.cpp
--- a/Engine/Plugins/Developer/ProjectLauncher/Source/ProjectLauncher/Private/Widgets/Profiles/SCustomLaunchCustomProfileSelector.cpp
+++ b/Engine/Plugins/Developer/ProjectLauncher/Source/ProjectLauncher/Private/Widgets/Profiles/SCustomLaunchCustomProfileSelector.cpp
@@ -331,7 +331,7 @@ const FSlateBrush* SCustomLaunchCustomProfileSelector::GetProfileImage(ILauncher
 	}
 	else
 	{
-		TArray<const PlatformInfo::FTargetPlatformInfo*> PlatformInfos = Model->GetPlatformInfos(Profile);
+		TArray<const PlatformInfo::FTargetPlatformInfo*> PlatformInfos = Model->GetPlatformInfosForAllCommands(Profile);
 		return FProjectLauncherStyle::GetProfileBrushForPlatforms(PlatformInfos, EPlatformIconSize::Large);
 	}
 }
@@ -535,6 +535,8 @@ void SCustomLaunchCustomProfileSelector::OnDeviceRemoved( const FString DeviceID
 
 void SCustomLaunchCustomProfileSelector::SetSelectedDevices( const TArray<FString> DeviceIDs, ILauncherProfilePtr Profile)
 {
+	check(Model->IsBasicLaunchProfile(Profile));
+
 	if (ensure(Profile->GetDeployedDeviceGroup() != nullptr))
 	{
 		Profile->GetDeployedDeviceGroup()->RemoveAllDevices();
@@ -542,7 +544,7 @@ void SCustomLaunchCustomProfileSelector::SetSelectedDevices( const TArray<FStrin
 		{
 			Profile->GetDeployedDeviceGroup()->AddDevice(DeviceID);
 		}
-		Model->UpdatedCookedPlatformsFromDeployDeviceProxy(Profile.ToSharedRef());
+		Model->UpdatedCookedPlatformsFromDeployDeviceProxy(Profile.ToSharedRef(), Profile->GetFirstBuildCookRun().ToSharedRef());
 
 		OnProfileModified.ExecuteIfBound(Profile);
 	}
diff --git a/Engine/Plugins/Developer/ProjectLauncher/Source/ProjectLauncher/Private/Widgets/Profiles/SCustomLaunchProfilesPanel.cpp b/Engine/Plugins/Developer/ProjectLauncher/Source/ProjectLauncher/Private/Widgets/Profiles/SCustomLaunchProfilesPanel.cpp
--- a/Engine/Plugins/Developer/ProjectLauncher/Source/ProjectLauncher/Private/Widgets/Profiles/SCustomLaunchProfilesPanel.cpp
+++ b/Engine/Plugins/Developer/ProjectLauncher/Source/ProjectLauncher/Private/Widgets/Profiles/SCustomLaunchProfilesPanel.cpp
@@ -385,6 +385,39 @@ TSharedRef<SWidget> SCustomLaunchProfilesPanel::CreateProfileEditorWidget()
 				]
 			]
 
+			+SHorizontalBox::Slot()
+			.AutoWidth()
+			.VAlign(VAlign_Center)
+			.HAlign(HAlign_Left)
+			.Padding(4,4)
+			[
+				SNew(SComboButton)
+				.ComboButtonStyle(FAppStyle::Get(), "SimpleComboButton")
+				.ToolTipText(LOCTEXT("CommandsToolTip", "Configure the commands to run."))
+				.OnGetMenuContent(PropertyEditor.ToSharedRef(), &SCustomLaunchCustomProfileEditor::MakeCommandsMenu)
+				.MenuPlacement(MenuPlacement_ComboBoxRight)
+				.ButtonContent()
+				[
+					SNew(SHorizontalBox)
+					+SHorizontalBox::Slot()
+					.AutoWidth()
+					.Padding(2, 0)
+					[
+						SNew(SImage)
+						.Image(FProjectLauncherStyle::Get().GetBrush("Icons.Extensions"))
+						.ColorAndOpacity(FSlateColor::UseForeground())
+					]
+					+SHorizontalBox::Slot()
+					.AutoWidth()
+					.Padding(2, 0)
+					[
+						SNew(STextBlock)
+						.Text(LOCTEXT("CommandsLabel", "Tasks"))
+						.ColorAndOpacity(FSlateColor::UseForeground())
+					]
+				]
+			]
+
 			+SHorizontalBox::Slot()
 			.FillWidth(1)
 			[
@@ -670,6 +703,11 @@ const FSlateBrush* SCustomLaunchProfilesPanel::GetLaunchButtonImage() const
 	{
 		return FProjectLauncherStyle::Get().GetBrush("Icons.Task.TestAutomation");
 	}
+	else if (Profile->GetBuildCookRunCommands().Num() != 1)
+	{
+		return FProjectLauncherStyle::Get().GetBrush("Icons.Task.Launch");
+	}	
+	// @fixme: figure out what the icon is like for multiple BCR steps
 	else if (Profile->GetLaunchMode() != ELauncherProfileLaunchModes::DoNotLaunch)
 	{
 		return FProjectLauncherStyle::Get().GetBrush("Icons.Task.Launch");
@@ -770,7 +808,7 @@ const FSlateBrush* SCustomLaunchProfilesPanel::GetSelectedProfileImage() const
 	}
 	else
 	{
-		TArray<const PlatformInfo::FTargetPlatformInfo*> PlatformInfos = Model->GetPlatformInfos(Model->GetSelectedProfile());
+		TArray<const PlatformInfo::FTargetPlatformInfo*> PlatformInfos = Model->GetPlatformInfosForAllCommands(Model->GetSelectedProfile());
 		return FProjectLauncherStyle::GetProfileBrushForPlatforms(PlatformInfos, EPlatformIconSize::Large);
 	}
 }
@@ -784,7 +822,7 @@ FText SCustomLaunchProfilesPanel::GetSelectedProfileName() const
 {
 	if (Model->IsBasicLaunchProfile(Model->GetSelectedProfile()))
 	{
-		const TSharedPtr<ITargetDeviceProxy> DeviceProxy = Model->GetDeviceProxy(Model->GetSelectedProfile().ToSharedRef());
+		const TSharedPtr<ITargetDeviceProxy> DeviceProxy = Model->GetDeviceProxy(Model->GetSelectedProfile()->GetFirstBuildCookRun().ToSharedRef());
 		if (DeviceProxy.IsValid())
 		{
 			return FText::FromString(DeviceProxy->GetName());
@@ -808,7 +846,7 @@ FText SCustomLaunchProfilesPanel::GetSelectedProfileDescription() const
 {
 	if (Model->IsBasicLaunchProfile(Model->GetSelectedProfile()))
 	{
-		const TSharedPtr<ITargetDeviceProxy> DeviceProxy = Model->GetDeviceProxy(Model->GetSelectedProfile().ToSharedRef());
+		const TSharedPtr<ITargetDeviceProxy> DeviceProxy = Model->GetDeviceProxy(Model->GetSelectedProfile()->GetFirstBuildCookRun().ToSharedRef());
 		if (DeviceProxy.IsValid())
 		{
 			return DeviceProxy->GetPlatformDisplayName(NAME_None);
@@ -865,7 +903,10 @@ void SCustomLaunchProfilesPanel::SetDefaultProjectPath(FString ProjectPath)
 
 	if (Model->GetSelectedProfile().IsValid())
 	{
-		Model->UpdateCookedPlatformsFromBuildTarget(Model->GetSelectedProfile().ToSharedRef());
+		for (const ILauncherProfileBuildCookRunRef& BuildCookRun : Model->GetSelectedProfile()->GetBuildCookRunCommands())
+		{
+			Model->UpdateCookedPlatformsFromBuildTarget(Model->GetSelectedProfile().ToSharedRef(), BuildCookRun);
+		}
 	}
 }
 
@@ -882,7 +923,10 @@ void SCustomLaunchProfilesPanel::SetDefaultBuildTarget(TArray<FString> BuildTarg
 
 	if (Model->GetSelectedProfile().IsValid())
 	{
-		Model->UpdateCookedPlatformsFromBuildTarget(Model->GetSelectedProfile().ToSharedRef());
+		for (const ILauncherProfileBuildCookRunRef& BuildCookRun : Model->GetSelectedProfile()->GetBuildCookRunCommands())
+		{
+			Model->UpdateCookedPlatformsFromBuildTarget(Model->GetSelectedProfile().ToSharedRef(), BuildCookRun);
+		}
 	}
 }
 
diff --git a/Engine/Plugins/Developer/ProjectLauncher/Source/ProjectLauncher/Public/Extension/AutomatedTestLaunchExtension.h b/Engine/Plugins/Developer/ProjectLauncher/Source/ProjectLauncher/Public/Extension/AutomatedTestLaunchExtension.h
--- a/Engine/Plugins/Developer/ProjectLauncher/Source/ProjectLauncher/Public/Extension/AutomatedTestLaunchExtension.h
+++ b/Engine/Plugins/Developer/ProjectLauncher/Source/ProjectLauncher/Public/Extension/AutomatedTestLaunchExtension.h
@@ -8,10 +8,6 @@
 
 namespace ProjectLauncher
 {
-
-
-
-
 	/**
 	 * Helper base class for an automated test launch extension
 	 * Creates an item in the extension menu that will toggle the test on and off
@@ -27,6 +23,8 @@ namespace ProjectLauncher
 		 */
 		UE_API bool IsTestActive() const;
 
+		UE_API virtual bool ManagesUATCommand( const ILauncherProfileUATCommandRef& UATCommand ) const override;
+
 	protected:
 		UE_API virtual void OnAdded() override;
 		UE_API virtual void OnRemoved() override;
@@ -70,18 +68,17 @@ namespace ProjectLauncher
 		UE_API virtual void InternalInitialize() override;
 
 	private:
+		UE_API virtual void OnCustomValidation() override;
 		virtual int32 GetHeadingSortOrder() const override { return GetTestOrder(); }
+		UE_API virtual FLaunchProfileTreeNode& AddDefaultHeadingInternal(FLaunchProfileTreeData& ProfileTreeData, FText DisplayName, int32 SortOrder) const override;
 	};
 
 
-
-
-
 	class FAutomatedTestLaunchExtension : public FLaunchExtension
 	{
 	public:
 		UE_API virtual void GetExtensionsMenuEntry( FExtensionsMenuEntry& MenuEntry ) const override;
-
+		virtual class FAutomatedTestLaunchExtension* AsAutomatedTestLaunchExtension() override { return this; }
 	};
 };
 
diff --git a/Engine/Plugins/Developer/ProjectLauncher/Source/ProjectLauncher/Public/Extension/BuildCookRunLaunchExtension.h b/Engine/Plugins/Developer/ProjectLauncher/Source/ProjectLauncher/Public/Extension/BuildCookRunLaunchExtension.h
new file mode 100644
--- /dev/null
+++ b/Engine/Plugins/Developer/ProjectLauncher/Source/ProjectLauncher/Public/Extension/BuildCookRunLaunchExtension.h
@@ -0,0 +1,215 @@
+// Copyright Epic Games, Inc. All Rights Reserved.
+
+#pragma once
+
+#include "Extension/LaunchExtension.h"
+
+#define UE_API PROJECTLAUNCHER_API
+
+namespace ProjectLauncher
+{
+	/**
+	 * Base for a launch extension for am ILauncherProfileBuildCookRun
+	 */
+	class FBuildCookRunExtension : public TSharedFromThis<FBuildCookRunExtension>
+	{
+	public:
+		struct FArgs
+		{
+			ILauncherProfileRef Profile;
+			ILauncherProfileBuildCookRunRef BuildCookRun;
+			ILauncherProfileBuildCookRunRef DefaultBuildCookRun;
+			TSharedRef<FModel> Model;
+			TSharedRef<FLaunchExtension> Extension;
+			TSharedRef<FLaunchExtensionInstance> ExtensionInstance;
+			FLaunchProfileTreeData& OwnerTreeData;
+		};
+
+		UE_API FBuildCookRunExtension( const FArgs& InArgs );
+		UE_API virtual ~FBuildCookRunExtension();
+
+		/**
+		 * Hook to allow the extension to add extra fields to the property editing tree, if the tree builder allows it
+		 * Property tree items should be hidden until the user has selected something to make it relevent, to avoid cluttering the UI.
+		 */
+		virtual void CustomizeTree( FLaunchProfileTreeNode& ProfileTreeNode ) {};
+
+		/**
+		 * Create the default heading for this extension
+		 */
+		UE_API FLaunchProfileTreeNode& AddDefaultHeading(FLaunchProfileTreeNode& ProfileTreeNode) const;
+
+		/**
+		 * Create the default heading for this extension, with a bespoke display name
+		 */
+		UE_API FLaunchProfileTreeNode& AddDefaultHeading(FLaunchProfileTreeNode& ProfileTreeNode, FText OverrideDisplayName) const;
+
+
+
+		/**
+		 * Advanced hook to allow any advanced modification of the command line when a UAT command test is launched
+		 */
+		virtual void CustomizeUATCommandLine( FString& InOutCommandLine ) {};
+
+		/**
+		 * Notification callback when a property has been changed in our profile
+		 */
+		void OnPropertyChanged() {};
+
+		/**
+		 * Notification callback when the project has been changed in our profile
+		 */
+		virtual void OnProjectChanged() {};
+
+		/**
+		 * Notification callback when the profile is validated
+		 */
+		virtual void OnValidate() {};
+
+		/** 
+		 * Get the extension we belong to
+		 */
+		inline TSharedRef<FLaunchExtension> GetExtension() const { return Extension; }
+
+		/** 
+		 * Get the extension instance that instantiated us
+		 */
+		inline TSharedRef<FLaunchExtensionInstance> GetExtensionInstance() const { return ExtensionInstance; }
+
+	protected:
+		/** 
+		 * Get the profile we were instantiated for
+		 */
+		inline ILauncherProfileRef GetProfile() const { return Profile; }
+
+		/** 
+		 * Get the BuildCookRun we were instantiated for
+		 */
+		inline ILauncherProfileBuildCookRunRef GetBuildCookRun() const { return BuildCookRun; }
+
+		/** 
+		 * Get the reference default BuildCookRun
+		 */
+		inline ILauncherProfileBuildCookRunRef GetDefaultBuildCookRun() const { return DefaultBuildCookRun; }
+		
+		/** 
+		 * Get the model, for general purpose helper functions
+		 */
+		inline TSharedRef<FModel> GetModel() const { return Model; }
+
+	private:
+		void OnUATCommandValidation( ILauncherProfileUATCommandRef UATCommand );
+		virtual int32 GetHeadingSortOrder() const { return 0; }
+
+		ILauncherProfileRef Profile;
+		ILauncherProfileBuildCookRunRef BuildCookRun;
+		ILauncherProfileBuildCookRunRef DefaultBuildCookRun;
+		TSharedRef<FModel> Model;
+		TSharedRef<FLaunchExtension> Extension;
+		TSharedRef<FLaunchExtensionInstance> ExtensionInstance;
+		FLaunchProfileTreeData& OwnerTreeData;
+		FPropertyChangedDelegate PropertyChangedDelegate;
+
+
+	};
+
+
+	/**
+	 * Interface for creating & maintaining build cook run extensions
+	 */
+	class IBuildCookRunExtensionFactory
+	{
+	public:
+
+		/**
+		 * A launch extension must override this to create a new build cook run extension
+		 */
+		virtual TSharedRef<FBuildCookRunExtension> CreateBuildCookRunExtension( const FBuildCookRunExtension::FArgs& InArgs ) = 0;
+
+		/**
+		 * Returns whether there are any active build cook run extensions from this launch extension
+		 */
+		virtual bool HasActiveBuildCookRunExtension() const = 0;
+
+		/**
+		 * Customizes the build cook run extensions menu to allow extenions to be enabled & disabled 
+		 */
+		virtual void CustomizeBuildCookRunExtensionMenu( FMenuBuilder& InMenuBuilder, const ILauncherProfileBuildCookRunRef& InBuildCookRun ) = 0;
+
+	protected:
+		/*
+		 * Returns whether build cook run extensions should be enabled by default. Typically false
+		 */
+		virtual bool IsBuildCookRunExtensionEnabledByDefault( const ILauncherProfileBuildCookRunRef& InBuildCookRun ) const = 0;
+
+		/*
+		 * Returns whether our build cook run extensions can be toggled on or off. Typically true
+		 */
+		virtual bool CanToggleBuildCookRunExtension( const ILauncherProfileBuildCookRunRef& InBuildCookRun, bool bWantToEnable ) const = 0;
+	};
+
+
+
+
+	class FBuildCookRunLaunchExtensionInstance : public FLaunchExtensionInstance, public IBuildCookRunExtensionFactory
+	{
+	public:
+
+		UE_API FBuildCookRunLaunchExtensionInstance( FArgs& InArgs);
+		UE_API virtual ~FBuildCookRunLaunchExtensionInstance();
+
+		UE_API virtual void InternalInitialize() override;
+
+
+		UE_API virtual IBuildCookRunExtensionFactory* AsBuildCookRunFactory() override { return this; }
+
+		UE_API bool HasActiveBuildCookRunExtension() const override;
+
+		UE_API void CustomizeBuildCookRunExtensionMenu( FMenuBuilder& InMenuBuilder, const ILauncherProfileBuildCookRunRef& InBuildCookRun )  override;
+
+		UE_API virtual TSharedRef<FBuildCookRunExtension> CreateBuildCookRunExtension( const FBuildCookRunExtension::FArgs& InArgs ) = 0;
+
+		UE_API virtual void CustomizeTree( FLaunchProfileTreeData& ProfileTreeData );
+		UE_API virtual void OnPropertyChanged() override;
+
+		// fixme: ideally enumerate FLaunchExtension to find those that are BCR
+		virtual bool CanBeRemoved() const override final { return true; }
+
+	protected:
+		/**
+		 * Get the key name to use for reading & writing a configuration value for the given BuildCookRun, expected to be passed as Name to the GetConfigXXX/SetConfigXXX functions
+		 * 
+		 * @param Config where the value is stored
+		 * @param Name name of the value
+		 * @param BuildCookRun the build cook run command
+		 * 
+		 * @returns namespaced version of Name
+		 */
+		UE_API FString GetConfigKeyName( const TCHAR* Name, const ILauncherProfileBuildCookRunRef& BuildCookRun) const;
+
+		UE_API virtual bool IsBuildCookRunExtensionEnabledByDefault( const ILauncherProfileBuildCookRunRef& InBuildCookRun ) const override;
+		UE_API virtual bool CanToggleBuildCookRunExtension( const ILauncherProfileBuildCookRunRef& InBuildCookRun, bool bWantToEnable ) const override;
+		UE_API virtual bool CustomizeUATCommandLine( const ILauncherProfileUATCommandRef& InUATCommand, FString& InOutCommandLine ) override;
+
+		UE_API void OnUATCommandAdded(const ILauncherProfileUATCommandRef& UATCommand);
+		UE_API void OnUATCommandRemoved(const ILauncherProfileUATCommandRef& UATCommand);
+		UE_API TSharedRef<FBuildCookRunExtension> CreateBuildCookRunExtensionInternal( const ILauncherProfileBuildCookRunRef& InBuildCookRun );
+
+		TMap<ILauncherProfileUATCommandPtr,TSharedRef<FBuildCookRunExtension>> BuildCookRunExtensions;
+
+		friend class FBuildCookRunExtension;
+	};
+
+
+
+	class FBuildCookRunLaunchExtension : public FLaunchExtension
+	{
+	public:
+		virtual bool IsAlwaysCreated( ILauncherProfileRef InProfile, TSharedRef<FModel> InModel ) const override { return true; }
+		virtual class FBuildCookRunLaunchExtension* AsBuildCookRunLaunchExtension() override { return this; }
+	};
+
+
+};
+
+#undef UE_API
diff --git a/Engine/Plugins/Developer/ProjectLauncher/Source/ProjectLauncher/Public/Extension/LaunchExtension.h b/Engine/Plugins/Developer/ProjectLauncher/Source/ProjectLauncher/Public/Extension/LaunchExtension.h
--- a/Engine/Plugins/Developer/ProjectLauncher/Source/ProjectLauncher/Public/Extension/LaunchExtension.h
+++ b/Engine/Plugins/Developer/ProjectLauncher/Source/ProjectLauncher/Public/Extension/LaunchExtension.h
@@ -13,6 +13,7 @@
 namespace ProjectLauncher
 {
 	class FLaunchExtension;
+	class FLaunchExtensionInstance;
 
 	typedef TMulticastDelegate<void()> FPropertyChangedDelegate;
 
@@ -71,7 +72,7 @@ namespace ProjectLauncher
 		 * Hook to allow the extension to add extra fields to the property editing tree, if the tree builder allows it
 		 * Property tree items should be hidden until the user has selected something to make it relevent, to avoid cluttering the UI.
 		 */
-		virtual void CustomizeTree( FLaunchProfileTreeData& ProfileTreeData ) {};
+		UE_API virtual void CustomizeTree( FLaunchProfileTreeData& ProfileTreeData );
 
 		/**
 		 * Advanced hook to allow any advanced modification of the command line when our profile is launched
@@ -91,12 +92,12 @@ namespace ProjectLauncher
 		/**
 		 * Advanced hook to allow any advanced modification of the command line when a UAT command test is launched
 		 */
-		virtual bool CustomizeUATCommandLine( const ILauncherProfileUATCommandRef& InUATCommand, FString& InOutCommandLine ) { return false; }
+		UE_API virtual bool CustomizeUATCommandLine( const ILauncherProfileUATCommandRef& InUATCommand, FString& InOutCommandLine );
 
 		/**
 		 * Notification callback when a property has been changed in our profile
 		 */
-		virtual void OnPropertyChanged() {};
+		UE_API virtual void OnPropertyChanged();
 
 		/**
 		 * Notification callback when the project has been changed in our profile
@@ -154,10 +155,26 @@ namespace ProjectLauncher
 		 */
 		inline FPropertyChangedDelegate& GetPropertyChangedDelegate() { return PropertyChangedDelegate; }
 
+		/**
+		 * Attempt to cast this launch extension instance to a build cook run factory
+		 */
+		virtual class IBuildCookRunExtensionFactory* AsBuildCookRunFactory() { return nullptr; }
+
+		/**
+		 * Determine if this launch extension instance manages the given UAT command
+		 */
+		virtual bool ManagesUATCommand( const ILauncherProfileUATCommandRef& UATCommand ) const { return false; }
+
 		/**
 		 * Internal intitialization function, called immediately after construction - don't call directly
 		 */
-		UE_API virtual void InternalInitialize() {};
+		UE_API virtual void InternalInitialize();
+
+
+		/**
+		 * Internal property tree customization function - don't call directly
+		 */
+		UE_API void CustomizeTree();
 
 		UE_API virtual void OnAdded() {};
 		UE_API virtual void OnRemoved() {};
@@ -201,6 +218,13 @@ namespace ProjectLauncher
 		 */
 		UE_API void SetCommandLine( const FString& CommandLine );
 
+		/**
+		 * Get the owner tree data
+		 */
+		inline FLaunchProfileTreeData& GetOwnerTreeData() const { return OwnerTreeData; }
+
+
+	public:
 		/**
 		 * Get the name of the platform to launch on. this is the highest priority platform/target in the list of platforms
 		 */
@@ -315,7 +339,17 @@ namespace ProjectLauncher
 		*/
 		UE_API void RequestTreeRefresh();
 
+		/** 
+		 * Request that the enture property tree is rebuilt - such as when adding or removing top-level sections
+		*/
+		UE_API void RequestFullTreeRebuild();
+
+	protected:
+		UE_API virtual FLaunchProfileTreeNode& AddDefaultHeadingInternal(FLaunchProfileTreeData& ProfileTreeData, FText DisplayName, int32 SortOrder) const;
+
 	private:
+
+		UE_API virtual void OnCustomValidation();
 		virtual int32 GetHeadingSortOrder() const { return 0; }
 
 
@@ -374,6 +408,22 @@ namespace ProjectLauncher
 		 */
 		virtual bool CanBeCreated( ILauncherProfileRef InProfile, TSharedRef<FModel> InModel ) const { return true; }
 
+
+		/**
+		 * Attempt to cast this launch extension to an automated test launch extension
+		 */
+		virtual class FAutomatedTestLaunchExtension* AsAutomatedTestLaunchExtension() { return nullptr; }
+
+		/**
+		 * Attempt to cast this launch extension to an UAT command launch extension
+		 */
+		virtual class FUATCommandLaunchExtension* AsUATCommandLaunchExtension() { return nullptr; }
+
+		/**
+		 * Attempt to cast this launch extension to an BuildCookRun launch extension
+		 */
+		virtual class FBuildCookRunLaunchExtension* AsBuildCookRunLaunchExtension() { return nullptr; }
+
 		/**
 		 * Instantiate all compatible extensions
 		 * 
@@ -425,6 +475,12 @@ namespace ProjectLauncher
 
 		virtual void MakeCustomExtensionSubmenu(FMenuBuilder& MenuBuilder, ILauncherProfileRef InProfile, TSharedRef<FModel> InModel) {}
 
+		/**
+		 * Destroy all extension instances for the given profile
+		 * 
+		 * @param InProfile the profile who's extension instances to return
+		 */
+		UE_API static void DestroyExtensionInstancesForProfile(ILauncherProfilePtr InProfile);
 
 	private:
 		/**
@@ -434,7 +490,6 @@ namespace ProjectLauncher
 		 */
 		virtual TSharedPtr<FLaunchExtensionInstance> CreateInstanceForProfile( FLaunchExtensionInstance::FArgs& InArgs ) = 0;
 	};
-
 };
 
 #undef UE_API
diff --git a/Engine/Plugins/Developer/ProjectLauncher/Source/ProjectLauncher/Public/Extension/UATCommandLaunchExtension.h b/Engine/Plugins/Developer/ProjectLauncher/Source/ProjectLauncher/Public/Extension/UATCommandLaunchExtension.h
--- a/Engine/Plugins/Developer/ProjectLauncher/Source/ProjectLauncher/Public/Extension/UATCommandLaunchExtension.h
+++ b/Engine/Plugins/Developer/ProjectLauncher/Source/ProjectLauncher/Public/Extension/UATCommandLaunchExtension.h
@@ -23,6 +23,9 @@ namespace ProjectLauncher
 		 */
 		UE_API bool IsUATCommandActive() const;
 
+
+		UE_API virtual bool ManagesUATCommand( const ILauncherProfileUATCommandRef& UATCommand ) const override;
+
 	protected:
 		UE_API virtual void OnAdded() override;
 		UE_API virtual void OnRemoved() override;
@@ -66,7 +69,9 @@ namespace ProjectLauncher
 		UE_API virtual void InternalInitialize() override;
 
 	private:
+		UE_API virtual void OnCustomValidation() override;
 		virtual int32 GetHeadingSortOrder() const override { return GetUATCommandOrder(); }
+		UE_API virtual FLaunchProfileTreeNode& AddDefaultHeadingInternal(FLaunchProfileTreeData& ProfileTreeData, FText DisplayName, int32 SortOrder) const override;
 	};
 
 
@@ -75,7 +80,7 @@ namespace ProjectLauncher
 	{
 	public:
 		UE_API virtual void GetExtensionsMenuEntry( FExtensionsMenuEntry& MenuEntry ) const override;
-
+		virtual class FUATCommandLaunchExtension* AsUATCommandLaunchExtension() override { return this; }
 	};
 
 };
diff --git a/Engine/Plugins/Developer/ProjectLauncher/Source/ProjectLauncher/Public/Model/ProjectLauncherModel.h b/Engine/Plugins/Developer/ProjectLauncher/Source/ProjectLauncher/Public/Model/ProjectLauncherModel.h
--- a/Engine/Plugins/Developer/ProjectLauncher/Source/ProjectLauncher/Public/Model/ProjectLauncherModel.h
+++ b/Engine/Plugins/Developer/ProjectLauncher/Source/ProjectLauncher/Public/Model/ProjectLauncherModel.h
@@ -27,6 +27,7 @@ struct FTargetInfo;
 namespace ProjectLauncher
 {
 	extern const FText ProjectLauncherText;
+	extern bool bUseClassicView;
 
 	/* Defines how content is stored and deployed for a profile */
 	enum class EContentScheme : uint8
@@ -201,11 +202,21 @@ namespace ProjectLauncher
 		/* Get the default custom launch profile - used for "reset to default" */
 		UE_API const ILauncherProfileRef GetDefaultCustomLaunchProfile() const;
 
+		/* Get the default build cook run for the basic profile - used for "reset to default" */
+		UE_API const ILauncherProfileBuildCookRunRef GetBasicDefaultBuildCookRun() const;
+
+		/* Get the default build cook run for the custom profile - used for "reset to default" */
+		UE_API const ILauncherProfileBuildCookRunRef GetCustomDefaultBuildCookRun() const;
+
+
 		/* Determine the content scheme for a profile */
-		UE_API EContentScheme DetermineProfileContentScheme(const ILauncherProfileRef& Profile) const;
+		UE_API EContentScheme DetermineProfileContentScheme(const ILauncherProfileBuildCookRunRef& BuildCookRun) const;
 
 		/* Apply a content scheme to a profile, setting cook/deploy/packaging options etc */
-		UE_API void SetProfileContentScheme(EContentScheme ContentScheme, const ILauncherProfileRef& Profile, bool bWantToCook = true, ELauncherProfileDeploymentModes::Type DefaultDeploymentMode = ELauncherProfileDeploymentModes::CopyToDevice);
+		UE_API void SetProfileContentScheme(EContentScheme ContentScheme, const ILauncherProfileBuildCookRunRef& BuildCookRun, bool bWantToCook = true, ELauncherProfileDeploymentModes::Type DefaultDeploymentMode = ELauncherProfileDeploymentModes::CopyToDevice);
+
+		/* Helper to add a BuildCookRun to the given profile with common initialized properties */
+		UE_API ILauncherProfileBuildCookRunRef CreateBuildCookRun( const ILauncherProfileRef& Profile );
 
 		/* Record a new log message from launching */
 		UE_API TSharedPtr<FLaunchLogMessage> AddLogMessage(const FString& InMessage, ELogVerbosity::Type InVerbosity = ELogVerbosity::Log);
@@ -223,31 +234,32 @@ namespace ProjectLauncher
 		UE_DEPRECATED(5.8, "Use GetPlatformInfos")
 		static UE_API const PlatformInfo::FTargetPlatformInfo* GetPlatformInfo(const ILauncherProfilePtr& Profile);
 		static UE_API const PlatformInfo::FTargetPlatformInfo* GetPlatformInfo(FName PlatformName, const FTargetInfo& BuildTargetInfo);
-		static UE_API TArray<const PlatformInfo::FTargetPlatformInfo*> GetPlatformInfos(const ILauncherProfilePtr& Profile);
+		static UE_API TArray<const PlatformInfo::FTargetPlatformInfo*> GetPlatformInfos(const ILauncherProfileBuildCookRunRef& BuildCookRun);
+		static UE_API TArray<const PlatformInfo::FTargetPlatformInfo*> GetPlatformInfosForAllCommands(const ILauncherProfilePtr& Profile);
 
 		/* Returns true if profile or platform name refers to host platform */
 		UE_DEPRECATED(5.8, "No longer in use. IsUsingRemotePlatform may be of use")
 		static UE_API bool IsHostPlatform(const ILauncherProfilePtr& Profile);
 		static UE_API bool IsHostPlatform(FName PlatformName);
-		static UE_API bool IsUsingRemotePlatform(const ILauncherProfilePtr& Profile);
+		static UE_API bool IsUsingRemotePlatform(const ILauncherProfileBuildCookRunRef& BuildCookRun);
 
 		/* Get build target info for a project path or profile */
 		UE_API FTargetInfo GetBuildTargetInfo(FString BuildTargetName, const FString& ProjectPath);
 		UE_DEPRECATED(5.8, "Use GetBuildTargetInfos")
 		UE_API FTargetInfo GetBuildTargetInfo(const ILauncherProfileRef& Profile);
-		UE_API TArray<FTargetInfo> GetBuildTargetInfos(const ILauncherProfileRef& Profile);
+		UE_API TArray<FTargetInfo> GetBuildTargetInfos(const ILauncherProfileRef& Profile, const ILauncherProfileBuildCookRunRef& BuildCookRun);
 
 		/* Get a device proxy for a profile */
-		static UE_API TSharedPtr<ITargetDeviceProxy> GetDeviceProxy(const ILauncherProfileRef& Profile);
+		static UE_API TSharedPtr<ITargetDeviceProxy> GetDeviceProxy(const ILauncherProfileBuildCookRunRef& BuildCookRun);
 
 		/* Update device proxies based on the current cook platforms */
 		UE_API void UpdateDeployDeviceProxyFromCookedPlatforms(const ILauncherProfileRef& Profile);
 
 		/* Update cooked platforms based on selected device proxy */
-		UE_API void UpdatedCookedPlatformsFromDeployDeviceProxy(const ILauncherProfileRef& Profile, TSharedPtr<ITargetDeviceProxy> DeviceProxy = nullptr);
+		UE_API void UpdatedCookedPlatformsFromDeployDeviceProxy(const ILauncherProfileRef& Profile, const ILauncherProfileBuildCookRunRef& BuildCookRun, TSharedPtr<ITargetDeviceProxy> DeviceProxy = nullptr);
 
 		/* Update cooked platforms based on build target */
-		UE_API void UpdateCookedPlatformsFromBuildTarget(const ILauncherProfileRef& Profile);
+		UE_API void UpdateCookedPlatformsFromBuildTarget(const ILauncherProfileRef& Profile, const ILauncherProfileBuildCookRunRef& BuildCookRun);
 
 		/* Utility functions for resolving platform names */
 		static UE_API FString GetVanillaPlatformName(const FString& PlatformName);
@@ -266,6 +278,9 @@ namespace ProjectLauncher
 		/* Returns true if extensions are enabled for project launcher */
 		UE_API bool AreExtensionsEnabled() const;
 
+		/* Returns true if we can use a simplified layout for this profile */
+		UE_API bool CanUseSimplifiedLayout(const ILauncherProfileRef& Profile);
+
 		/* Get available maps for a project or engine */
 		UE_API TArray<FString> GetAvailableProjectMapNames(const FString& InProjectPath, bool bIncludeNonContentDirMaps = false);
 		UE_API TArray<FString> GetAvailableProjectMapPaths(const FString& InProjectPath, bool bIncludeNonContentDirMaps = false);
@@ -293,6 +308,9 @@ namespace ProjectLauncher
 		/* Get map paths for a project, caching results */
 		UE_API TArray<FString> GetAndCacheMapPaths(const FString& InOptionalProjectPath, bool bIncludeNonContentDirMaps = false);
 
+		/* Set the defaults for this BuildCookRun */
+		void SetDefaults(const ILauncherProfileRef& Profile, const ILauncherProfileBuildCookRunRef& BuildCookRun);
+
 	private:
 		FString ConfigFileName;
 		
diff --git a/Engine/Plugins/Developer/ProjectLauncher/Source/ProjectLauncher/Public/ProfileTree/GenericProfileTreeBuilder.h b/Engine/Plugins/Developer/ProjectLauncher/Source/ProjectLauncher/Public/ProfileTree/GenericProfileTreeBuilder.h
--- a/Engine/Plugins/Developer/ProjectLauncher/Source/ProjectLauncher/Public/ProfileTree/GenericProfileTreeBuilder.h
+++ b/Engine/Plugins/Developer/ProjectLauncher/Source/ProjectLauncher/Public/ProfileTree/GenericProfileTreeBuilder.h
@@ -4,6 +4,7 @@
 
 #include "ProfileTree/ILaunchProfileTreeBuilder.h"
 #include "Model/ProjectLauncherModel.h"
+#include "ILauncherProfileBuildCookRun.h"
 #include "Containers/Map.h"
 #include "Containers/Set.h"
 
@@ -48,143 +49,172 @@ namespace ProjectLauncher
 	protected:
 		// default property creation functions
 		UE_API void AddProjectProperty( FLaunchProfileTreeNode& HeadingNode );
-		UE_API void AddTargetProperty( FLaunchProfileTreeNode& HeadingNode );
-		UE_API void AddPlatformProperty( FLaunchProfileTreeNode& HeadingNode );
-		UE_API void AddConfigurationProperty( FLaunchProfileTreeNode& HeadingNode );
-		UE_API void AddContentSchemeProperty( FLaunchProfileTreeNode& HeadingNode );
-		UE_API void AddCompressPakFilesProperty( FLaunchProfileTreeNode& HeadingNode );
-		UE_API void AddUseIoStoreProperty( FLaunchProfileTreeNode& HeadingNode );
-		UE_API void AddGenerateChunksProperty( FLaunchProfileTreeNode& HeadingNode );
-		UE_API void AddZenSnapshotProperty(FLaunchProfileTreeNode& HeadingNode);
-		UE_API void AddImportZenSnapshotProperty( FLaunchProfileTreeNode& HeadingNode );
-		UE_API void AddZenPakStreamingPathProperty( FLaunchProfileTreeNode& HeadingNode );
-		UE_API void AddSubmissionPackageProperty( FLaunchProfileTreeNode& HeadingNode );
-		UE_API void AddIncrementalCookProperty( FLaunchProfileTreeNode& HeadingNode );
-		UE_API void AddCookProperty( FLaunchProfileTreeNode& HeadingNode );
-		UE_API void AddBuildProperty( FLaunchProfileTreeNode& HeadingNode );
-		UE_API void AddForceBuildProperty( FLaunchProfileTreeNode& HeadingNode );
-		UE_API void AddBuildUATProperty( FLaunchProfileTreeNode& HeadingNode );
-		UE_API void AddArchitectureProperty( FLaunchProfileTreeNode& HeadingNode );
-		UE_API void AddStagingDirectoryProperty( FLaunchProfileTreeNode& HeadingNode );
-		UE_API void AddArchiveBuildProperty( FLaunchProfileTreeNode& HeadingNode );
-		UE_API void AddArchiveBuildDirectoryProperty( FLaunchProfileTreeNode& HeadingNode );
-		UE_API void AddDeployProperty( FLaunchProfileTreeNode& HeadingNode );
-		UE_API void AddIncrementalDeployProperty( FLaunchProfileTreeNode& HeadingNode );
-		UE_API void AddTargetDeviceProperty( FLaunchProfileTreeNode& HeadingNode );
-		UE_API void AddRunProperty( FLaunchProfileTreeNode& HeadingNode );
-		UE_API void AddInitialMapProperty( FLaunchProfileTreeNode& HeadingNode );
-		UE_API void AddCommandLineProperty( FLaunchProfileTreeNode& HeadingNode );
-
-		// helper callbacks to simplify control enable/visibility
-		FLaunchProfileTreeNode::FGetBool ForPak;
-		FLaunchProfileTreeNode::FGetBool ForZen;
-		FLaunchProfileTreeNode::FGetBool ForZenWS;
-		FLaunchProfileTreeNode::FGetBool ForCooked;
-		FLaunchProfileTreeNode::FGetBool ForEnabledCooked;
-		FLaunchProfileTreeNode::FGetBool ForContent;
-		FLaunchProfileTreeNode::FGetBool ForCode;
-		FLaunchProfileTreeNode::FGetBool ForDeployment;
-		FLaunchProfileTreeNode::FGetBool ForRun;
-		FLaunchProfileTreeNode::FGetBool ForSubmissionPackage;
-		FLaunchProfileTreeNode::FGetString EmptyString;
-
-
-
-
-		TArray<FString> CachedDeployDeviceIDs;
-		enum class EDeployDeviceOption : uint8
-		{
-			Default,
-			Selected,
-		};
 
 		UE_API FString GetProjectPath() const;
 		UE_API void SetProjectName(FString ProjectPath);
 		UE_API bool HasProject() const;
 
-		UE_API TArray<FString> GetBuildTargets() const;
-		UE_API void SetBuildTargets( TArray<FString> BuildTargets);
-
-		UE_API void SetBuildConfiguration(EBuildConfiguration BuildConfiguration);
-		UE_API EBuildConfiguration GetBuildConfiguration() const;
-
-		UE_API void SetContentScheme(EContentScheme ContentScheme);
-		UE_API bool IsContentSchemeAvailable(EContentScheme, FText& OutReason) const;
-
-		UE_API FString GetCommandLine(EBuildTargetType BuildTarget) const;
-		UE_API void SetCommandLine( const FString& NewCommandLine, EBuildTargetType BuildTarget );
-
-		UE_API void SetSelectedPlatforms( TArray<FString> SelectedPlatforms );
-		UE_API TArray<FString> GetSelectedPlatforms() const;
-
-		UE_API void SetCook( bool bCook );
-		UE_API bool GetCook( ILauncherProfilePtr InProfile = nullptr ) const;
-
-		UE_API void SetDeployDeviceIDs(TArray<FString> DeployDeviceIDs);
-		UE_API TArray<FString> GetDeployDeviceIDs() const;
-		UE_API EDeployDeviceOption GetDeployDeviceOption() const;
-		UE_API void SetDeployDeviceOption( EDeployDeviceOption DeployDeviceOption );
-		UE_API float GetDeployDeviceListHeight() const;
-		UE_API void SetDeployDeviceListHeight( float NewHeight );
-		UE_API void OnDeviceRemoved(FString DeviceID);
-		UE_API TSharedRef<SWidget> CreateDeployDeviceWidget();
-		float DeployDeviceListHeight = 150.0f;
-
-		UE_API void SetBuild( bool bBuild );
-		UE_API bool GetBuild( ILauncherProfilePtr InProfile = nullptr ) const;
-
-		UE_API void SetForceBuild( bool bForceBuild );
-		UE_API bool GetForceBuild( ILauncherProfilePtr InProfile = nullptr ) const;
-
-		UE_API void SetArchitecture( FString Architecture, EBuildTargetType TargetType );
-		UE_API FString GetArchitecture(EBuildTargetType TargetType) const;
-		UE_API FText GetArchitectureDisplayName( FString Architecture );
-
-		UE_API void SetDeployToDevice( bool bDeployToDevice );
-		UE_API bool GetDeployToDevice( ILauncherProfilePtr InProfile = nullptr ) const;
-
-		UE_API void SetIsRunning( bool bRun );
-		UE_API bool GetIsRunning( ILauncherProfilePtr InProfile = nullptr ) const;
-
-		UE_API void OnInitialMapChanged( TSharedPtr<FString> InitialMap, ESelectInfo::Type );
-		UE_API TSharedPtr<FString> GetInitialMap() const;
 
+		class FBuildCookRun : public TSharedFromThis<FBuildCookRun>
+		{
+		public:
+			FBuildCookRun( FGenericProfileTreeBuilder& InOwner, ILauncherProfileBuildCookRunRef InBuildCookRun, const TSharedRef<FModel>& InModel );
+			~FBuildCookRun();
+
+			UE_API void AddTargetProperty( FLaunchProfileTreeNode& HeadingNode );
+			UE_API void AddPlatformProperty( FLaunchProfileTreeNode& HeadingNode );
+			UE_API void AddConfigurationProperty( FLaunchProfileTreeNode& HeadingNode );
+			UE_API void AddContentSchemeProperty( FLaunchProfileTreeNode& HeadingNode );
+			UE_API void AddCompressPakFilesProperty( FLaunchProfileTreeNode& HeadingNode );
+			UE_API void AddUseIoStoreProperty( FLaunchProfileTreeNode& HeadingNode );
+			UE_API void AddGenerateChunksProperty( FLaunchProfileTreeNode& HeadingNode );
+			UE_API void AddZenSnapshotProperty(FLaunchProfileTreeNode& HeadingNode);
+			UE_API void AddImportZenSnapshotProperty( FLaunchProfileTreeNode& HeadingNode );
+			UE_API void AddZenPakStreamingPathProperty( FLaunchProfileTreeNode& HeadingNode );
+			UE_API void AddSubmissionPackageProperty( FLaunchProfileTreeNode& HeadingNode );
+			UE_API void AddIncrementalCookProperty( FLaunchProfileTreeNode& HeadingNode );
+			UE_API void AddCookProperty( FLaunchProfileTreeNode& HeadingNode );
+			UE_API void AddBuildProperty( FLaunchProfileTreeNode& HeadingNode );
+			UE_API void AddForceBuildProperty( FLaunchProfileTreeNode& HeadingNode );
+			UE_API void AddArchitectureProperty( FLaunchProfileTreeNode& HeadingNode );
+			UE_API void AddStagingDirectoryProperty( FLaunchProfileTreeNode& HeadingNode );
+			UE_API void AddArchiveBuildProperty( FLaunchProfileTreeNode& HeadingNode );
+			UE_API void AddArchiveBuildDirectoryProperty( FLaunchProfileTreeNode& HeadingNode );
+			UE_API void AddDeployProperty( FLaunchProfileTreeNode& HeadingNode );
+			UE_API void AddIncrementalDeployProperty( FLaunchProfileTreeNode& HeadingNode );
+			UE_API void AddTargetDeviceProperty( FLaunchProfileTreeNode& HeadingNode );
+			UE_API void AddRunProperty( FLaunchProfileTreeNode& HeadingNode );
+			UE_API void AddInitialMapProperty( FLaunchProfileTreeNode& HeadingNode );
+			UE_API void AddCommandLineProperty( FLaunchProfileTreeNode& HeadingNode );
+
+
+			enum class EDeployDeviceOption : uint8
+			{
+				Default,
+				Selected,
+			};
+
+			UE_API void Construct();
+
+
+			UE_API TArray<FString> GetBuildTargets() const;
+			UE_API void SetBuildTargets( TArray<FString> BuildTargets);
+
+			UE_API void SetBuildConfiguration(EBuildConfiguration BuildConfiguration);
+			UE_API EBuildConfiguration GetBuildConfiguration() const;
+
+			UE_API void SetContentScheme(EContentScheme ContentScheme);
+			UE_API bool IsContentSchemeAvailable(EContentScheme, FText& OutReason) const;
+
+			UE_API FString GetCommandLine(EBuildTargetType BuildTarget) const;
+			UE_API void SetCommandLine( const FString& NewCommandLine, EBuildTargetType BuildTarget );
+
+			UE_API void SetSelectedPlatforms( TArray<FString> SelectedPlatforms );
+			UE_API TArray<FString> GetSelectedPlatforms() const;
+
+			UE_API void SetCook( bool bCook );
+			UE_API bool GetCook( ILauncherProfileBuildCookRunPtr InBuildCookRun = nullptr ) const;
+
+			UE_API void SetDeployDeviceIDs(TArray<FString> DeployDeviceIDs);
+			UE_API TArray<FString> GetDeployDeviceIDs() const;
+			UE_API EDeployDeviceOption GetDeployDeviceOption() const;
+			UE_API void SetDeployDeviceOption( EDeployDeviceOption DeployDeviceOption );
+			UE_API float GetDeployDeviceListHeight() const;
+			UE_API void SetDeployDeviceListHeight( float NewHeight );
+			UE_API void OnDeviceRemoved(FString DeviceID);
+			UE_API TSharedRef<SWidget> CreateDeployDeviceWidget();
+
+			UE_API void SetBuild( bool bBuild );
+			UE_API bool GetBuild( ILauncherProfileBuildCookRunPtr InBuildCookRun = nullptr ) const;
+
+			UE_API void SetForceBuild( bool bForceBuild );
+			UE_API bool GetForceBuild( ILauncherProfileBuildCookRunPtr InBuildCookRun = nullptr ) const;
+
+			UE_API void SetArchitecture( FString Architecture, EBuildTargetType TargetType );
+			UE_API FString GetArchitecture(EBuildTargetType TargetType) const;
+			UE_API FText GetArchitectureDisplayName( FString Architecture );
+
+			UE_API void SetDeployToDevice( bool bDeployToDevice );
+			UE_API bool GetDeployToDevice( ILauncherProfileBuildCookRunPtr InBuildCookRun = nullptr ) const;
+
+			UE_API void SetIsRunning( bool bRun );
+			UE_API bool GetIsRunning( ILauncherProfileBuildCookRunPtr InBuildCookRun = nullptr ) const;
+
+			UE_API void OnInitialMapChanged( TSharedPtr<FString> InitialMap, ESelectInfo::Type );
+			UE_API TSharedPtr<FString> GetInitialMap() const;
+
+
+			UE_API void SetHasAdvancedPlatformTargets( bool bEnable );
+			UE_API bool HasAdvancedPlatformTargets() const;
+
+			UE_API void RefreshContentScheme();
+			UE_API void ValidateArchitectures();
+			UE_API void OnValidateProfile(ILauncherProfileUATCommandRef UATCommand);
+
+			UE_API void CacheArchitectures();
+			UE_API void CacheBuildTargets();
+
+			FText GetZenSnapshotText() const;
+
+			// helper callbacks to simplify control enable/visibility
+			FLaunchProfileTreeNode::FGetBool ForPak;
+			FLaunchProfileTreeNode::FGetBool ForZen;
+			FLaunchProfileTreeNode::FGetBool ForZenWS;
+			FLaunchProfileTreeNode::FGetBool ForCooked;
+			FLaunchProfileTreeNode::FGetBool ForEnabledCooked;
+			FLaunchProfileTreeNode::FGetBool ForContent;
+			FLaunchProfileTreeNode::FGetBool ForCode;
+			FLaunchProfileTreeNode::FGetBool ForDeployment;
+			FLaunchProfileTreeNode::FGetBool ForRun;
+			FLaunchProfileTreeNode::FGetBool ForSubmissionPackage;
+			FLaunchProfileTreeNode::FGetString EmptyString;
+
+			TArray<FString> CachedDeployDeviceIDs;
+
+			float DeployDeviceListHeight = 150.0f;
+
+			EContentScheme ContentScheme;
+			bool bHasAdvancedPlatformTargets;
+			bool bShouldCook;
+			TSharedPtr<SCustomLaunchPlatformCombo> PlatformCombo;
+			TSharedPtr<SCustomLaunchBuildTargetCombo> BuildTargetCombo;
+			TSharedPtr<SSearchableComboBox> InitalMapCombo;
+			TSharedPtr<SCustomLaunchDeviceListView> DeployDeviceListView;
+			TMap<EBuildTargetType, TArray<FString>> CachedBuildTargetArchitectures;
+			TSet<EBuildTargetType> CachedBuildTargetTypes;
 
-		UE_API void SetHasAdvancedPlatformTargets( bool bEnable );
-		UE_API bool HasAdvancedPlatformTargets() const;
+			FGenericProfileTreeBuilder& Owner;
+			ILauncherProfileBuildCookRunRef BuildCookRun;
+			const TSharedRef<FModel> Model;
+		};
 
 
 	protected:
 
+		void OnUATCommandAdded(const ILauncherProfileUATCommandRef& UATCommand);
+		void OnUATCommandRemoved(const ILauncherProfileUATCommandRef& UATCommand);
 		UE_API virtual void OnPropertyChanged() override;
-		UE_API virtual void OnValidateProfile();
-		UE_API void RefreshContentScheme();
 		UE_API void CacheStartupMapList() const;
-		UE_API void CacheArchitectures();
-		UE_API void CacheBuildTargets();
-		UE_API void ValidateArchitectures();
 		UE_API TSharedRef<SWidget> OnGenerateComboWidget( TSharedPtr<FString> InComboString );
-		FText GetZenSnapshotText() const;
 
 
+		mutable bool bStartupMapCacheDirty = true;
+		mutable TArray<TSharedPtr<FString>> CachedStartupMaps;
+
 		FLaunchProfileTreeDataRef TreeData;
 		const ILauncherProfileRef Profile;
 		const ILauncherProfileRef DefaultProfile;
+		ILauncherProfileBuildCookRunPtr DefaultBuildCookRun;
 		EProfileType ProfileType;
 
 		const TSharedRef<FModel> Model;
 
-		EContentScheme ContentScheme;
-		bool bHasAdvancedPlatformTargets;
-		bool bShouldCook;
-		mutable bool bStartupMapCacheDirty = true;
-		mutable TArray<TSharedPtr<FString>> CachedStartupMaps;
-		TSharedPtr<SCustomLaunchPlatformCombo> PlatformCombo;
-		TSharedPtr<SCustomLaunchBuildTargetCombo> BuildTargetCombo;
-		TSharedPtr<SSearchableComboBox> InitalMapCombo;
-		TSharedPtr<SCustomLaunchDeviceListView> DeployDeviceListView;
-		TMap<EBuildTargetType, TArray<FString>> CachedBuildTargetArchitectures;
-		TSet<EBuildTargetType> CachedBuildTargetTypes;
+		FBuildCookRun& Get(const ILauncherProfileBuildCookRunRef& BuildCookRun);
+
+
+
+
+		TMap<ILauncherProfileBuildCookRunRef,TSharedRef<FBuildCookRun>> BuildCookRunData;
 	};
 }
 
diff --git a/Engine/Plugins/Developer/ProjectLauncher/Source/ProjectLauncher/Public/ProfileTree/LaunchProfileTreeData.h b/Engine/Plugins/Developer/ProjectLauncher/Source/ProjectLauncher/Public/ProfileTree/LaunchProfileTreeData.h
--- a/Engine/Plugins/Developer/ProjectLauncher/Source/ProjectLauncher/Public/ProfileTree/LaunchProfileTreeData.h
+++ b/Engine/Plugins/Developer/ProjectLauncher/Source/ProjectLauncher/Public/ProfileTree/LaunchProfileTreeData.h
@@ -53,10 +53,10 @@ namespace ProjectLauncher
 		bool IsSet() const;
 
 		/* Returns true if the given profile has one of the associated errors. */
-		bool HasError( ILauncherProfileRef Profile ) const;
+		bool HasError( ILauncherProfileRef Profile, ILauncherProfileUATCommandPtr UATCommand ) const;
 
 		/* Collects error text from the profile for display. */
-		void GetErrorText( ILauncherProfileRef Profile, TArray<FString>& ErrorLines ) const;
+		void GetErrorText( ILauncherProfileRef Profile, ILauncherProfileUATCommandPtr UATCommand, TArray<FString>& ErrorLines ) const;
 
 	private:
 		bool bIsSet = false;                                             // whether this validation has been configured
@@ -148,18 +148,22 @@ namespace ProjectLauncher
 		UE_API FLaunchProfileTreeNode& AddSubHeading(const TCHAR* InName, FText InDisplayName, int32 SortOrder = 0);
 		UE_API FLaunchProfileTreeNode& GetSubHeading(const TCHAR* InName) const;
 
+		UE_API FLaunchProfileTreeNode& FindOrAddSubHeading(ILauncherProfileUATCommandRef UATCommand);
 
+		UE_API bool ContainsNodeRecursive( const FLaunchProfileTreeNodePtr& TargetNode ) const;
 
 		/* Node details */
 		FText Name;													// label for the node
 		int32 SortOrder;											// sort order for headings
 		int32 Depth = 0;											// tree depth
 		TSharedPtr<SWidget> Widget;									// UI widget
+		ILauncherProfileUATCommandPtr UATCommand;					// option UAT command associated with this node
 		FCallbacks Callbacks;										// generic callbacks
 		TArray<FLaunchProfileTreeNodePtr> Children;					// child nodes
 		TMap<FString,FLaunchProfileTreeNodePtr> SubHeadingNodes;    // named subheadings
 
 		const FLaunchProfileTreeData* GetTreeData() const { return TreeData; }
+		FLaunchProfileTreeData* GetTreeData() { return TreeData; }
 
 	protected:
 		/* Notify parent tree that a property has changed. */
@@ -183,6 +187,9 @@ namespace ProjectLauncher
 		UE_API FLaunchProfileTreeNode& AddHeading(const TCHAR* InName, FText InDisplayName, int32 SortOrder = 0);
 		UE_API FLaunchProfileTreeNode& GetHeading(const TCHAR* InName) const;
 
+		UE_API FLaunchProfileTreeNode& FindOrAddHeading(ILauncherProfileUATCommandRef UATCommand, bool bIsPrimaryHeading);
+		UE_API void SortNodes();
+
 		/* Invoke extension instances to build their UI. */
 		UE_API void CreateExtensionsUI();
 
@@ -199,6 +206,7 @@ namespace ProjectLauncher
 		ILaunchProfileTreeBuilder* TreeBuilder;                          // UI builder interface
 		TArray<TSharedPtr<FLaunchExtensionInstance>> ExtensionInstances; // extension instances attached
 		TSet<FLaunchProfileTreeNodePtr> RequestExpandHeadings;           // headings to expand
+		TSet<FLaunchProfileTreeNodePtr> UATCommandHeadings;              // headings that have additional controls for modifying the command
 		bool bRequestTreeRefresh = false;                                // request to refresh tree UI
 		bool bRequestTreeFullRebuild = false;                            // request to fully recreate the profile tree
 	};
diff --git a/Engine/Plugins/Performance/AutomatedPerfTesting/Source/AutomatedPerfTesting/Private/LaunchExtension/AutomatedPerfTestLaunchExtension.cpp b/Engine/Plugins/Performance/AutomatedPerfTesting/Source/AutomatedPerfTesting/Private/LaunchExtension/AutomatedPerfTestLaunchExtension.cpp
--- a/Engine/Plugins/Performance/AutomatedPerfTesting/Source/AutomatedPerfTesting/Private/LaunchExtension/AutomatedPerfTestLaunchExtension.cpp
+++ b/Engine/Plugins/Performance/AutomatedPerfTesting/Source/AutomatedPerfTesting/Private/LaunchExtension/AutomatedPerfTestLaunchExtension.cpp
@@ -126,10 +126,6 @@ const FString FAutomatedPerfTestLaunchExtensionInstance::GetTestInternalName() c
 void FAutomatedPerfTestLaunchExtensionInstance::CustomizeTree(ProjectLauncher::FLaunchProfileTreeData& ProfileTreeData)
 {
 	using namespace ProjectLauncher;
-	const auto IsVisible = [this]()
-	{
-		return IsTestActive() || GetConfigBool(EConfig::PerProfile, TEXT("Enabled") );
-	};
 
 	const auto GetDisplayName = [](EAutomatedPerfTestType Type)
 	{
@@ -170,25 +166,21 @@ void FAutomatedPerfTestLaunchExtensionInstance::CustomizeTree(ProjectLauncher::F
 		{
 			.GetValue = [this]() { return GetConfigBool(ConfigTypeCommon, EnableCSVProfilerConfig, true); },
 			.SetValue = [this](bool bEnable) { return SetConfigBool(ConfigTypeCommon, EnableCSVProfilerConfig, bEnable); },
-			.IsVisible = IsVisible,
 		})
 		.AddBoolean(LOCTEXT("WindowedLabel","Windowed"),
 		{
 			.GetValue = [this]()          { return GetConfigBool(ConfigTypeCommon, WindowedConfig);},
 			.SetValue = [this](bool bVal) { SetConfigBool(ConfigTypeCommon, WindowedConfig, bVal ); },
-			.IsVisible = IsVisible,
 		})
 		.AddInteger(LOCTEXT("IterationsLabel", "Iterations"), 
 		{
 			.GetValue = [this]() { return GetConfigInteger(ConfigTypeCommon, IterationsConfig, DefaultIterationCount); },
 			.SetValue = [this](int32 Value) { SetConfigInteger(ConfigTypeCommon, IterationsConfig, Value); },
-			.IsVisible = IsVisible,
 		})
 		.AddWidget(LOCTEXT("TestTypeLabel", "Test Type"),
 		{
 			.IsDefault = [this]() { return GetTestType() == DefaultTestType; },
 			.SetToDefault = [this]() { SetTestType(DefaultTestType); },
-			.IsVisible = IsVisible
 		},
 		SNew(SCustomLaunchCombo<EAutomatedPerfTestType>)
 		.OnSelectionChanged(this, &FAutomatedPerfTestLaunchExtensionInstance::OnTestTypeSelectionChanged)
diff --git a/Engine/Plugins/Performance/AutomatedPerfTesting/Source/AutomatedPerfTesting/Private/LaunchExtension/AutomatedPerfTestLaunchExtension.h b/Engine/Plugins/Performance/AutomatedPerfTesting/Source/AutomatedPerfTesting/Private/LaunchExtension/AutomatedPerfTestLaunchExtension.h
--- a/Engine/Plugins/Performance/AutomatedPerfTesting/Source/AutomatedPerfTesting/Private/LaunchExtension/AutomatedPerfTestLaunchExtension.h
+++ b/Engine/Plugins/Performance/AutomatedPerfTesting/Source/AutomatedPerfTesting/Private/LaunchExtension/AutomatedPerfTestLaunchExtension.h
@@ -51,7 +51,7 @@ private:
 	{
 		const auto IsTestType = [this]()
 		{
-			const bool bEnabled = IsTestActive() || GetConfigBool(EConfig::PerProfile, TEXT("Enabled"));
+			const bool bEnabled = IsTestActive();
 			return GetTestType() == TestType && bEnabled;
 		};
 
