Splitting out command line variable editing into a separate extension type

#rnx
#Synced-CL 49676362

diff --git a/Engine/Plugins/Developer/ProjectLauncher/Source/CommonLaunchExtensions/Private/Insights/InsightsLaunchExtension.cpp b/Engine/Plugins/Developer/ProjectLauncher/Source/CommonLaunchExtensions/Private/Insights/InsightsLaunchExtension.cpp
--- a/Engine/Plugins/Developer/ProjectLauncher/Source/CommonLaunchExtensions/Private/Insights/InsightsLaunchExtension.cpp
+++ b/Engine/Plugins/Developer/ProjectLauncher/Source/CommonLaunchExtensions/Private/Insights/InsightsLaunchExtension.cpp
@@ -13,8 +13,24 @@ const TCHAR* FInsightsLaunchExtensionInstance::TraceParam = TEXT("-trace=");
 const TCHAR* FInsightsLaunchExtensionInstance::StatNamedEventsParam = TEXT("-statnamedevents");
 
 
+TSharedRef<ProjectLauncher::FCmdLineParametersExtension> FInsightsLaunchExtensionInstance::Create( const ProjectLauncher::FCmdLineParametersExtension::FArgs& InArgs )
+{
+	return MakeShared<FInsightsCommandLineExtension>(InArgs);
+}
+
+bool FInsightsLaunchExtensionInstance::SupportsCmdLineType( const FString& Type ) const
+{
+	return ICmdLineParametersExtensionFactory::DefaultSupportsType(Type);
+}
+
+
+
+FInsightsLaunchExtensionInstance::FInsightsCommandLineExtension::FInsightsCommandLineExtension( const ProjectLauncher::FCmdLineParametersExtension::FArgs& InArgs )
+	: ProjectLauncher::FCmdLineParametersExtension(InArgs)
+{
+}
 
-bool FInsightsLaunchExtensionInstance::GetExtensionParameters( TArray<FString>& OutParameters ) const
+bool FInsightsLaunchExtensionInstance::FInsightsCommandLineExtension::GetExtensionParameters( TArray<FString>& OutParameters ) const
 {
 	OutParameters.Add(HostParam);
 	OutParameters.Add(FileParam);
@@ -23,7 +39,7 @@ bool FInsightsLaunchExtensionInstance::GetExtensionParameters( TArray<FString>&
 	return true;
 }
 
-FText FInsightsLaunchExtensionInstance::GetExtensionParameterDisplayName( const FString& InParameter ) const
+FText FInsightsLaunchExtensionInstance::FInsightsCommandLineExtension::GetExtensionParameterDisplayName( const FString& InParameter ) const
 {
 	if (InParameter == HostParam)
 	{
@@ -38,12 +54,12 @@ FText FInsightsLaunchExtensionInstance::GetExtensionParameterDisplayName( const
 		return LOCTEXT("TraceNamedEventsParam", "Capture named events");
 	}
 
-	return FLaunchExtensionInstance::GetExtensionParameterDisplayName(InParameter);
+	return ProjectLauncher::FCmdLineParametersExtension::GetExtensionParameterDisplayName(InParameter);
 
 }
 
 
-void FInsightsLaunchExtensionInstance::CacheTraceChannels()
+void FInsightsLaunchExtensionInstance::FInsightsCommandLineExtension::CacheTraceChannels()
 {
 	TraceChannels.Reset();
 
@@ -53,7 +69,7 @@ void FInsightsLaunchExtensionInstance::CacheTraceChannels()
 	TraceChannels.Sort();
 }
 
-void FInsightsLaunchExtensionInstance::ToggleTraceChannel( const FString& InChannel )
+void FInsightsLaunchExtensionInstance::FInsightsCommandLineExtension::ToggleTraceChannel( const FString& InChannel )
 {
 	if (TraceChannels.Contains(InChannel))
 	{
@@ -77,13 +93,13 @@ void FInsightsLaunchExtensionInstance::ToggleTraceChannel( const FString& InChan
 	SetCommandLine(CommandLine);
 }
 
-bool FInsightsLaunchExtensionInstance::IsTraceChannelEnabled( const FString InChannel ) const
+bool FInsightsLaunchExtensionInstance::FInsightsCommandLineExtension::IsTraceChannelEnabled( const FString InChannel ) const
 {
 	return TraceChannels.Contains(InChannel);
 }
 
 
-void FInsightsLaunchExtensionInstance::CustomizeParametersSubmenu( FMenuBuilder& MenuBuilder )
+void FInsightsLaunchExtensionInstance::FInsightsCommandLineExtension::CustomizeParametersSubmenu( FMenuBuilder& MenuBuilder )
 {
 #if UE_TRACE_ENABLED || UE_TRACE_MINIMAL_ENABLED
 	auto ChannelsMenuDelegate = [this]( FMenuBuilder& MenuBuilder )
diff --git a/Engine/Plugins/Developer/ProjectLauncher/Source/CommonLaunchExtensions/Private/Insights/InsightsLaunchExtension.h b/Engine/Plugins/Developer/ProjectLauncher/Source/CommonLaunchExtensions/Private/Insights/InsightsLaunchExtension.h
--- a/Engine/Plugins/Developer/ProjectLauncher/Source/CommonLaunchExtensions/Private/Insights/InsightsLaunchExtension.h
+++ b/Engine/Plugins/Developer/ProjectLauncher/Source/CommonLaunchExtensions/Private/Insights/InsightsLaunchExtension.h
@@ -3,32 +3,43 @@
 #pragma once
 
 #include "Extension/LaunchExtension.h"
+#include "Extension/CmdLineParametersExtension.h"
 
 
-
-class FInsightsLaunchExtensionInstance : public ProjectLauncher::FLaunchExtensionInstance
+class FInsightsLaunchExtensionInstance : public ProjectLauncher::FLaunchExtensionInstance, public ProjectLauncher::ICmdLineParametersExtensionFactory
 {
 public:
 	FInsightsLaunchExtensionInstance( FArgs& InArgs ) : FLaunchExtensionInstance(InArgs) {};
 	virtual ~FInsightsLaunchExtensionInstance() = default;
 
-	virtual bool GetExtensionParameters( TArray<FString>& OutParameters ) const override;
-	virtual FText GetExtensionParameterDisplayName( const FString& InParameter ) const override;
-	
-	virtual void CustomizeParametersSubmenu( FMenuBuilder& MenuBuilder ) override;
+	virtual class ProjectLauncher::ICmdLineParametersExtensionFactory* AsCmdLineParametersFactory() override { return this; }
 
-	void CacheTraceChannels();
-	void ToggleTraceChannel( const FString& InChannel );
-	bool IsTraceChannelEnabled( const FString InChannel ) const;
+	// ... ICmdLineParametersExtensionFactory
+	virtual TSharedRef<ProjectLauncher::FCmdLineParametersExtension> Create( const ProjectLauncher::FCmdLineParametersExtension::FArgs& InArgs ) override;
+	virtual bool SupportsCmdLineType( const FString& Type ) const override;
 
 private:
-	TArray<FString> TraceChannels;
+	class FInsightsCommandLineExtension : public ProjectLauncher::FCmdLineParametersExtension
+	{
+	public:
+		FInsightsCommandLineExtension( const ProjectLauncher::FCmdLineParametersExtension::FArgs& InArgs );
+
+		virtual bool GetExtensionParameters( TArray<FString>& OutParameters ) const override;
+		virtual FText GetExtensionParameterDisplayName( const FString& InParameter ) const override;
+	
+		virtual void CustomizeParametersSubmenu( FMenuBuilder& MenuBuilder ) override;
+
+		void CacheTraceChannels();
+		void ToggleTraceChannel( const FString& InChannel );
+		bool IsTraceChannelEnabled( const FString InChannel ) const;
+
+		TArray<FString> TraceChannels;
+	};
 
 	static const TCHAR* FileParam;
 	static const TCHAR* HostParam;
 	static const TCHAR* TraceParam;
 	static const TCHAR* StatNamedEventsParam;
-
 };
 
 
diff --git a/Engine/Plugins/Developer/ProjectLauncher/Source/ProjectLauncher/Private/Extension/CmdLineParametersExtension.cpp b/Engine/Plugins/Developer/ProjectLauncher/Source/ProjectLauncher/Private/Extension/CmdLineParametersExtension.cpp
new file mode 100644
--- /dev/null
+++ b/Engine/Plugins/Developer/ProjectLauncher/Source/ProjectLauncher/Private/Extension/CmdLineParametersExtension.cpp
@@ -0,0 +1,176 @@
+// Copyright Epic Games, Inc. All Rights Reserved.
+
+#include "Extension/CmdLineParametersExtension.h"
+#include "Utils/LauncherCmdLineUtils.h"
+#include "Model/ProjectLauncherModel.h"
+
+#define LOCTEXT_NAMESPACE "FCmdLineParametersExtensionInstance"
+
+
+namespace ProjectLauncher
+{
+	FCmdLineParametersExtension::FCmdLineParametersExtension( const FArgs& InArgs )
+		: Profile(InArgs.ExtensionInstance->GetProfile())
+		, Model(InArgs.Model)
+		, Extension(InArgs.ExtensionInstance->GetExtension())
+		, ExtensionInstance(InArgs.ExtensionInstance)
+		, SetCommandLineDelegate(InArgs.SetCommandLine)
+		, GetCommandLineDelegate(InArgs.GetCommandLine)
+		, TypeName(InArgs.TypeName)
+	{
+		check(SetCommandLineDelegate.IsSet());
+		check(GetCommandLineDelegate.IsSet());
+
+		Profile->OnProjectChanged().AddRaw( this, &FCmdLineParametersExtension::OnProjectChanged);
+	}
+
+	FCmdLineParametersExtension::~FCmdLineParametersExtension()
+	{
+		Profile->OnProjectChanged().RemoveAll(this);
+	}
+
+	FText FCmdLineParametersExtension::GetExtensionParameterDisplayName( const FString& InParameter ) const
+	{
+		return FText::FromString(InParameter);
+	}
+
+	bool FCmdLineParametersExtension::GetExtensionParameters( TArray<FString>& OutParameters ) const
+	{
+		return false;
+	}
+
+	void FCmdLineParametersExtension::SetCommandLine( const FString& InCommandLine )
+	{
+		SetCommandLineDelegate(InCommandLine);
+		ExtensionInstance->BroadcastPropertyChanged();
+	}
+
+	FString FCmdLineParametersExtension::GetCommandLine() const
+	{
+		return GetCommandLineDelegate();
+	}
+
+
+	void FCmdLineParametersExtension::MakeCommandLineSubmenu( FMenuBuilder& MenuBuilder )
+	{
+		// create the menu entry
+		auto ToggleParameter = [this]( FString Parameter )
+		{
+			FString CommandLine = GetCommandLine();
+			if (CmdLineUtils::IsParameterUsed(CommandLine, Parameter))
+			{
+				CmdLineUtils::RemoveParameter(CommandLine, Parameter);
+			}
+			else
+			{
+				CmdLineUtils::AddParameter(CommandLine, Parameter);
+			}
+			SetCommandLine(CommandLine);
+		};
+
+		TArray<FString> Parameters;
+		if (GetExtensionParameters(Parameters) && Parameters.Num() > 0)
+		{
+			for (const FString& Parameter : Parameters)
+			{
+				FText ParameterDisplayName = GetExtensionParameterDisplayName(Parameter);
+				FText ParameterToolTip = FText::FromString(Parameter);
+
+				MenuBuilder.AddMenuEntry( 
+					ParameterDisplayName,
+					ParameterToolTip,
+					FSlateIcon(),
+					FUIAction(
+						FExecuteAction::CreateLambda(ToggleParameter, Parameter),
+						FCanExecuteAction(),
+						FIsActionChecked::CreateLambda( [this, Parameter]() { return CmdLineUtils::IsParameterUsed(GetCommandLine(), Parameter); } )
+					),
+					NAME_None,
+					EUserInterfaceActionType::Check);
+			}
+		}
+
+		CustomizeParametersSubmenu(MenuBuilder);
+	}
+
+
+
+
+	static TMap<uint64,TArray<TSharedRef<FCmdLineParametersExtension>>> GCmdLineParametersExtensions;
+
+	void ICmdLineParametersExtensionFactory::MakeCmdLineParametersSubmenu(FMenuBuilder& MenuBuilder, uint64 InstanceID, const FString& TypeName, TSharedRef<FModel> Model, TSharedRef<FLaunchExtensionInstance> ExtensionInstance, const FLaunchProfileTreeNode::FSetString& SetValue, const FLaunchProfileTreeNode::FGetString& GetValue )
+	{
+		// get the command line instances for this field
+		if (!GCmdLineParametersExtensions.Contains(InstanceID))
+		{
+			GCmdLineParametersExtensions.Add(InstanceID);
+		}
+		TArray<TSharedRef<FCmdLineParametersExtension>>& CmdLineExtensions = GCmdLineParametersExtensions[InstanceID];
+
+		// find or create the extension instance for this field
+		TSharedRef<FCmdLineParametersExtension>* CmdLineExtensionPtr = CmdLineExtensions.FindByPredicate( [ExtensionInstance]( const TSharedRef<FCmdLineParametersExtension>& CmdLineParametersInstance )
+		{
+			return CmdLineParametersInstance->GetExtensionInstance() == ExtensionInstance;
+		});
+		TSharedPtr<FCmdLineParametersExtension> CmdLineExtension;
+		if (CmdLineExtensionPtr)
+		{
+			CmdLineExtension = (*CmdLineExtensionPtr);
+		}
+		else
+		{
+			// create & cache a new extension instance
+			ProjectLauncher::ICmdLineParametersExtensionFactory* CmdLineParametersFactory = ExtensionInstance->AsCmdLineParametersFactory();
+			check(CmdLineParametersFactory); // should have been verified before calling
+
+			FCmdLineParametersExtension::FArgs Args
+			{
+				.Model = Model,
+				.ExtensionInstance = ExtensionInstance,
+				.SetCommandLine = SetValue,
+				.GetCommandLine = GetValue,
+				.TypeName = TypeName,
+			};
+
+			CmdLineExtension = CmdLineParametersFactory->Create(Args);
+			CmdLineExtensions.Add(CmdLineExtension.ToSharedRef());
+		}
+
+		// create the menu
+		CmdLineExtension->MakeCommandLineSubmenu(MenuBuilder);
+	}
+
+
+	void ICmdLineParametersExtensionFactory::DestroyCmdLineParametersExtensions(TSharedPtr<FLaunchExtensionInstance> Instance)
+	{
+		// remove all launch command line instances that come from the given extension, recording any instances that become empty
+		TSet<uint64> EmptyInstanceIDs;
+		for (TTuple<uint64,TArray<TSharedRef<FCmdLineParametersExtension>>>& Pair : GCmdLineParametersExtensions)
+		{
+			Pair.Value.RemoveAll( [Instance]( const TSharedRef<FCmdLineParametersExtension>& CmdLineParametersExtension )
+			{
+				return CmdLineParametersExtension->GetExtensionInstance() == Instance;
+			});
+
+			if (Pair.Value.Num() == 0)
+			{
+				EmptyInstanceIDs.Add(Pair.Key);
+			}
+		}
+
+		// remove all empty instances
+		for (uint64 EmptyInstanceID : EmptyInstanceIDs)
+		{
+			GCmdLineParametersExtensions.Remove(EmptyInstanceID);
+		}
+	}
+
+	bool ICmdLineParametersExtensionFactory::DefaultSupportsType( const FString& Type )
+	{
+		return Type.Equals(ProjectLauncher::CmdLineType::Launch);
+	}
+
+
+}
+
+#undef LOCTEXT_NAMESPACE
diff --git a/Engine/Plugins/Developer/ProjectLauncher/Source/ProjectLauncher/Private/Extension/LaunchExtension.cpp b/Engine/Plugins/Developer/ProjectLauncher/Source/ProjectLauncher/Private/Extension/LaunchExtension.cpp
--- a/Engine/Plugins/Developer/ProjectLauncher/Source/ProjectLauncher/Private/Extension/LaunchExtension.cpp
+++ b/Engine/Plugins/Developer/ProjectLauncher/Source/ProjectLauncher/Private/Extension/LaunchExtension.cpp
@@ -1,6 +1,7 @@
 // Copyright Epic Games, Inc. All Rights Reserved.
 
 #include "Extension/LaunchExtension.h"
+#include "Extension/CmdLineParametersExtension.h"
 #include "Model/ProjectLauncherModel.h"
 #include "Utils/LauncherCmdLineUtils.h"
 #include "Containers/Array.h"
@@ -172,7 +173,20 @@ namespace ProjectLauncher
 	void RemoveExtensionInstancesForProfile( const ILauncherProfileRef& Profile )
 	{
 		FScopeLock Lock(&GExtensionsCS);
-		GExtensionInstances.Remove(Profile.ToSharedPtr());
+
+		if (GExtensionInstances.Contains(Profile))
+		{
+			TArray<TSharedPtr<FLaunchExtensionInstance>>& ExtensionInstances = GExtensionInstances.FindChecked(Profile.ToSharedPtr());
+			for (TSharedPtr<FLaunchExtensionInstance>& ExtensionInstance : ExtensionInstances)
+			{
+				if (ExtensionInstance->GetProfile() == Profile)
+				{
+					ICmdLineParametersExtensionFactory::DestroyCmdLineParametersExtensions(ExtensionInstance);
+				}
+			}
+
+			GExtensionInstances.Remove(Profile.ToSharedPtr());
+		}
 
 		for (TTuple<TSharedPtr<FLaunchExtension>,TArray<TSharedPtr<FLaunchExtensionInstance>>>& Pair : GInstancesPerExtension)
 		{
@@ -475,17 +489,21 @@ namespace ProjectLauncher
 
 	FString FLaunchExtensionInstance::GetParameterValue( const FString& InParameter ) const
 	{
+		PRAGMA_DISABLE_DEPRECATION_WARNINGS
 		return CmdLineUtils::GetParameterValue(GetCommandLine(), InParameter);
+		PRAGMA_ENABLE_DEPRECATION_WARNINGS
 	}
 
 	bool FLaunchExtensionInstance::UpdateParameterValue( const FString& InParameter, const FString& NewValue )
 	{
+		PRAGMA_DISABLE_DEPRECATION_WARNINGS
 		FString CommandLine = GetCommandLine();
 		if (CmdLineUtils::UpdateParameterValue(CommandLine, InParameter, NewValue))
 		{
 			SetCommandLine(CommandLine);
 			return true;
 		}
+		PRAGMA_ENABLE_DEPRECATION_WARNINGS
 
 		return false;
 	}
@@ -494,81 +512,50 @@ namespace ProjectLauncher
 
 	FString FLaunchExtensionInstance::GetFinalParameter( const FString& InParameter ) const
 	{
+		PRAGMA_DISABLE_DEPRECATION_WARNINGS
 		return CmdLineUtils::GetFinalParameter(GetCommandLine(), InParameter);
+		PRAGMA_ENABLE_DEPRECATION_WARNINGS
 	}
 
 	void FLaunchExtensionInstance::AddParameter( const FString& InParameter )
 	{
+		PRAGMA_DISABLE_DEPRECATION_WARNINGS
 		FString CommandLine = GetCommandLine();
 		CmdLineUtils::AddParameter(CommandLine, InParameter);
 		SetCommandLine(CommandLine);
+		PRAGMA_ENABLE_DEPRECATION_WARNINGS
 	}
 
 	void FLaunchExtensionInstance::RemoveParameter( const FString& InParameter )
 	{
+		PRAGMA_DISABLE_DEPRECATION_WARNINGS
 		FString CommandLine = GetCommandLine();
 		if (CmdLineUtils::RemoveParameter(CommandLine, InParameter))
 		{
 			SetCommandLine(CommandLine);
 		}
+		PRAGMA_ENABLE_DEPRECATION_WARNINGS
 	}
 
 
 	bool FLaunchExtensionInstance::IsParameterUsed( const FString& InParameter ) const
 	{
+		PRAGMA_DISABLE_DEPRECATION_WARNINGS
 		return CmdLineUtils::IsParameterUsed(GetCommandLine(), InParameter);
+		PRAGMA_ENABLE_DEPRECATION_WARNINGS
 	}
 
 	void FLaunchExtensionInstance::SetParameterUsed( const FString& InParameter, bool bUsed )
 	{
+		PRAGMA_DISABLE_DEPRECATION_WARNINGS
 		FString CommandLine = GetCommandLine();
 		CmdLineUtils::SetParameterUsed(CommandLine, InParameter, bUsed);
 		SetCommandLine(CommandLine);
+		PRAGMA_ENABLE_DEPRECATION_WARNINGS
 	}
 
 
 
-	void FLaunchExtensionInstance::MakeCommandLineSubmenu( FMenuBuilder& MenuBuilder)
-	{
-		auto ToggleParameter = [this]( FString Parameter )
-		{
-			FString CommandLine = GetCommandLine();
-			if (CmdLineUtils::IsParameterUsed(CommandLine, Parameter))
-			{
-				CmdLineUtils::RemoveParameter(CommandLine, Parameter);
-			}
-			else
-			{
-				CmdLineUtils::AddParameter(CommandLine, Parameter);
-			}
-			SetCommandLine(CommandLine);
-		};
-
-		TArray<FString> Parameters;
-		if (GetExtensionParameters(Parameters) && Parameters.Num() > 0)
-		{
-			for (const FString& Parameter : Parameters)
-			{
-				FText ParameterDisplayName = GetExtensionParameterDisplayName(Parameter);
-				FText ParameterToolTip = FText::FromString(Parameter);
-
-				MenuBuilder.AddMenuEntry( 
-					ParameterDisplayName,
-					ParameterToolTip,
-					FSlateIcon(),
-					FUIAction(
-						FExecuteAction::CreateLambda(ToggleParameter, Parameter),
-						FCanExecuteAction(),
-						FIsActionChecked::CreateLambda( [this, Parameter]() { return CmdLineUtils::IsParameterUsed(GetCommandLine(), Parameter); } )
-					),
-					NAME_None,
-					EUserInterfaceActionType::Check);
-			}
-		}
-
-		CustomizeParametersSubmenu(MenuBuilder);
-	}
-
 	FString FLaunchExtensionInstance::GetCommandLine() const
 	{
 		return Profile->GetAdditionalCommandLineParameters();
diff --git a/Engine/Plugins/Developer/ProjectLauncher/Source/ProjectLauncher/Private/Model/ProjectLauncherModel.cpp b/Engine/Plugins/Developer/ProjectLauncher/Source/ProjectLauncher/Private/Model/ProjectLauncherModel.cpp
--- a/Engine/Plugins/Developer/ProjectLauncher/Source/ProjectLauncher/Private/Model/ProjectLauncherModel.cpp
+++ b/Engine/Plugins/Developer/ProjectLauncher/Source/ProjectLauncher/Private/Model/ProjectLauncherModel.cpp
@@ -30,6 +30,12 @@ namespace ProjectLauncher
 	const FText ProjectLauncherText = LOCTEXT("ProjectLauncherTitle", "Project Launcher");
 	bool bUseClassicView = true;
 
+	namespace CmdLineType
+	{
+		const TCHAR* Launch = TEXT("Launch");
+		const TCHAR* UAT = TEXT("UAT");
+	}
+
 	FModel::FModel(const TSharedRef<ITargetDeviceProxyManager>& InDeviceProxyManager, const TSharedRef<ILauncher>& InLauncher, const TSharedRef<ILauncherProfileManager>& InProfileManager)
 		: DeviceProxyManager(InDeviceProxyManager)
 		, Launcher(InLauncher)
diff --git a/Engine/Plugins/Developer/ProjectLauncher/Source/ProjectLauncher/Private/ProfileTree/LaunchProfileTreeData.cpp b/Engine/Plugins/Developer/ProjectLauncher/Source/ProjectLauncher/Private/ProfileTree/LaunchProfileTreeData.cpp
--- a/Engine/Plugins/Developer/ProjectLauncher/Source/ProjectLauncher/Private/ProfileTree/LaunchProfileTreeData.cpp
+++ b/Engine/Plugins/Developer/ProjectLauncher/Source/ProjectLauncher/Private/ProfileTree/LaunchProfileTreeData.cpp
@@ -3,6 +3,7 @@
 #include "ProfileTree/LaunchProfileTreeData.h"
 #include "ProfileTree/ILaunchProfileTreeBuilder.h"
 #include "Extension/LaunchExtension.h"
+#include "Extension/CmdLineParametersExtension.h"
 #include "Styling/ProjectLauncherStyle.h"
 #include "Styling/AppStyle.h"
 #include "Styling/StyleColors.h"
@@ -619,13 +620,21 @@ namespace ProjectLauncher
 		);
 	}
 
-	FLaunchProfileTreeNode& FLaunchProfileTreeNode::AddCommandLineString(FText InName, FStringCallbacks&& StringCallbacks)
+	FLaunchProfileTreeNode& FLaunchProfileTreeNode::AddCommandLineString(FText InName, FStringCallbacks&& StringCallbacks, const TCHAR* InType)
 	{
 		if (!TreeData->Model->AreExtensionsEnabled())
 		{
 			return AddString(InName, MoveTemp(StringCallbacks));
 		}
 
+		if (InType == nullptr)
+		{
+			InType = ProjectLauncher::CmdLineType::Launch;
+		}
+
+		static std::atomic<uint64> UniqueInstanceID = 0;
+		uint64 InstanceID = UniqueInstanceID.fetch_add(1, std::memory_order_relaxed) + 1;;
+
 		check(StringCallbacks.GetValue);
 		check(StringCallbacks.SetValue);
 
@@ -660,7 +669,7 @@ namespace ProjectLauncher
 			return FText::FromString(StringCallbacks.GetValue());
 		};
 
-		auto OnGetCmdlineParameterMenuContent = [this,StringCallbacks]()
+		auto OnGetCmdlineParameterMenuContent = [this, StringCallbacks, InstanceID, TypeName = FString(InType)]()
 		{
 			const bool bShouldCloseWindowAfterMenuSelection = false;
 			const bool bCloseSelfOnly = false;
@@ -669,18 +678,29 @@ namespace ProjectLauncher
 
 			for (TSharedPtr<FLaunchExtensionInstance> ExtensionInstance : TreeData->ExtensionInstances)
 			{
+				ProjectLauncher::ICmdLineParametersExtensionFactory* CmdLineParametersFactory = ExtensionInstance->AsCmdLineParametersFactory();
+				if (CmdLineParametersFactory == nullptr || !CmdLineParametersFactory->SupportsCmdLineType(TypeName))
+				{
+					continue;
+				}
+
+				auto MakeCommandLineSubmenu = [this, StringCallbacks, InstanceID, TypeName, ExtensionInstance]( FMenuBuilder& MenuBuilder )
+				{
+					ICmdLineParametersExtensionFactory::MakeCmdLineParametersSubmenu(MenuBuilder, InstanceID, TypeName, GetTreeData()->Model, ExtensionInstance.ToSharedRef(), StringCallbacks.SetValue, StringCallbacks.GetValue);
+				};
+
 #if 1
 				// flat list of extensions - easier to see but risks getting cluttered later on
 				MenuBuilder.BeginSection(NAME_None, ExtensionInstance->GetExtension()->GetDisplayName());
-				ExtensionInstance->MakeCommandLineSubmenu(MenuBuilder);
+				MakeCommandLineSubmenu(MenuBuilder);
 				MenuBuilder.EndSection();
 #else 
 				// extensions in submenus - more scalable, but harder to see everything at a glance
 				const bool bInOpenSubMenuOnClick = false;
 				MenuBuilder.AddSubMenu(
-					ExtensionInstance->GetExtension()->GetDisplayName()),
+					ExtensionInstance->GetExtension()->GetDisplayName(),
 					FText::GetEmpty(),
-					FNewMenuDelegate::CreateSP( ExtensionInstance.ToSharedRef(), &FLaunchExtensionInstance::MakeCommandLineSubmenu),
+					FNewMenuDelegate::CreateLambda(MakeCommandLineSubmenu),
 					bInOpenSubMenuOnClick,
 					FSlateIcon(),
 					bShouldCloseWindowAfterMenuSelection
diff --git a/Engine/Plugins/Developer/ProjectLauncher/Source/ProjectLauncher/Public/Extension/CmdLineParametersExtension.h b/Engine/Plugins/Developer/ProjectLauncher/Source/ProjectLauncher/Public/Extension/CmdLineParametersExtension.h
new file mode 100644
--- /dev/null
+++ b/Engine/Plugins/Developer/ProjectLauncher/Source/ProjectLauncher/Public/Extension/CmdLineParametersExtension.h
@@ -0,0 +1,133 @@
+// Copyright Epic Games, Inc. All Rights Reserved.
+
+#pragma once
+
+#include "Extension/LaunchExtension.h"
+#include "Templates/Function.h"
+
+#define UE_API PROJECTLAUNCHER_API
+
+namespace ProjectLauncher
+{
+	/**
+	 * Base for a command line extension
+	 */
+	class FCmdLineParametersExtension : public TSharedFromThis<FCmdLineParametersExtension>
+	{
+	public:
+		struct FArgs
+		{
+			TSharedRef<FModel> Model;
+			TSharedRef<FLaunchExtensionInstance> ExtensionInstance;
+			FLaunchProfileTreeNode::FSetString SetCommandLine;
+			FLaunchProfileTreeNode::FGetString GetCommandLine;
+			FString TypeName;
+		};
+
+		UE_API FCmdLineParametersExtension( const FArgs& InArgs );
+		UE_API virtual ~FCmdLineParametersExtension();
+
+		/**
+		 * Returns the parameters that this extension provides. They will be added to the submenu.
+		 */
+		UE_API virtual bool GetExtensionParameters( TArray<FString>& OutParameters ) const;
+
+		/**
+		 * Returns the user-facing name for the given parameter. It will default to the parameter itself.
+		 */
+		UE_API virtual FText GetExtensionParameterDisplayName( const FString& InParameter ) const;
+
+		/**
+		 * Creates the submenu for picking command line parameters
+		 */
+		UE_API void MakeCommandLineSubmenu( FMenuBuilder& MenuBuilder );
+
+		/** 
+		 * Get the extension we belong to
+		 */
+		inline TSharedRef<FLaunchExtension> GetExtension() const { return Extension; }
+
+		/** 
+		 * Get the extension instance that instantiated us
+		 */
+		inline TSharedRef<FLaunchExtensionInstance> GetExtensionInstance() const { return ExtensionInstance; }
+
+	protected:
+		/**
+		 * Hook to allow the extension to extend the extension parameters menu
+		 */
+		virtual void CustomizeParametersSubmenu( FMenuBuilder& MenuBuilder ) {};
+
+		/**
+		 * Notification callback when the project has been changed in our profile
+		 */
+		virtual void OnProjectChanged() {};
+
+		/** 
+		 * Get the profile we were instantiated for
+		 */
+		inline ILauncherProfileRef GetProfile() const { return Profile; }
+	
+		/** 
+		 * Get the model, for general purpose helper functions
+		 */
+		inline TSharedRef<FModel> GetModel() const { return Model; }
+
+		/**
+		 * Update the command line
+		 */
+		UE_API void SetCommandLine( const FString& InCommandLine );
+
+		/**
+		 * Get the current command line
+		 */
+		UE_API FString GetCommandLine() const;
+
+	private:
+		ILauncherProfileRef Profile;
+		TSharedRef<FModel> Model;
+		TSharedRef<FLaunchExtension> Extension;
+		TSharedRef<FLaunchExtensionInstance> ExtensionInstance;
+		FLaunchProfileTreeNode::FSetString SetCommandLineDelegate;
+		FLaunchProfileTreeNode::FGetString GetCommandLineDelegate;
+		FString TypeName;
+	};
+
+
+	/**
+	 * Interface for creating & maintaining command line extensions
+	 */
+	class ICmdLineParametersExtensionFactory
+	{
+	public:
+		virtual ~ICmdLineParametersExtensionFactory() = default;
+
+		/**
+		 * A launch extension must override this to create a command line extension
+		 */
+		virtual TSharedRef<FCmdLineParametersExtension> Create( const FCmdLineParametersExtension::FArgs& InArgs ) = 0;
+
+		/**
+		 * Whether the extension supports the given type. Either a custom string or one of ProjectLauncher::CmdLineType::
+		 */
+		virtual bool SupportsCmdLineType( const FString& Type ) const = 0;
+
+
+		/**
+		 * Create the command line parameters menu for the given extension
+		 */
+		static UE_API void MakeCmdLineParametersSubmenu(FMenuBuilder& MenuBuilder, uint64 InstanceID, const FString& TypeName, TSharedRef<FModel> Model, TSharedRef<FLaunchExtensionInstance> Instance, const FLaunchProfileTreeNode::FSetString& SetValue, const FLaunchProfileTreeNode::FGetString& GetValue );
+
+		/**
+		 * Remove any command line extensions from the given launch extension
+		 */
+		static UE_API void DestroyCmdLineParametersExtensions(TSharedPtr<FLaunchExtensionInstance> Instance);
+
+		/**
+		 * Helper function for interface implementers - returns True for Launch command lines (the most common case)
+		 */
+		static UE_API bool DefaultSupportsType( const FString& Type );
+	};
+};
+
+#undef UE_API
diff --git a/Engine/Plugins/Developer/ProjectLauncher/Source/ProjectLauncher/Public/Extension/LaunchExtension.h b/Engine/Plugins/Developer/ProjectLauncher/Source/ProjectLauncher/Public/Extension/LaunchExtension.h
--- a/Engine/Plugins/Developer/ProjectLauncher/Source/ProjectLauncher/Public/Extension/LaunchExtension.h
+++ b/Engine/Plugins/Developer/ProjectLauncher/Source/ProjectLauncher/Public/Extension/LaunchExtension.h
@@ -46,11 +46,13 @@ namespace ProjectLauncher
 		/**
 		 * Returns the parameters that this extension provides. They will be added to the submenu.
 		 */
+		UE_DEPRECATED(5.8, "refactor code to use ICmdLineParametersExtension")
 		UE_API virtual bool GetExtensionParameters( TArray<FString>& OutParameters ) const;
 
 		/**
 		 * Returns the user-facing name for the given parameter. It will default to the parameter itself.
 		 */
+		UE_DEPRECATED(5.8, "refactor code to use ICmdLineParametersExtension")
 		UE_API virtual FText GetExtensionParameterDisplayName( const FString& InParameter ) const;
 
 		/**
@@ -66,6 +68,7 @@ namespace ProjectLauncher
 		/**
 		 * Hook to allow the extension to extend the extension parameters menu
 		 */
+		UE_DEPRECATED(5.8, "refactor code to use ICmdLineParametersExtension")
 		virtual void CustomizeParametersSubmenu( FMenuBuilder& MenuBuilder ) {};
 
 		/**
@@ -160,6 +163,11 @@ namespace ProjectLauncher
 		 */
 		virtual class IBuildCookRunExtensionFactory* AsBuildCookRunFactory() { return nullptr; }
 
+		/**
+		 * Attempt to cast this launch extension instance to a launch command line factory
+		 */
+		virtual class ICmdLineParametersExtensionFactory* AsCmdLineParametersFactory() { return nullptr; }
+
 		/**
 		 * Determine if this launch extension instance manages the given UAT command
 		 */
@@ -204,18 +212,18 @@ namespace ProjectLauncher
 
 		/** 
 		 * Get the current command line
-		 * Extensions should generally go through this function instead of querying the Profile directly
 		 *
 		 * @returns the current command line string
 		 */
+		UE_DEPRECATED(5.8, "access the appropriate field from a BuildCookRun instance, or use ICmdLineParametersExtension")
 		UE_API FString GetCommandLine() const;
 
 		/** 
 		 * Update the current command line
-		 * Extensions should generally go through this function instead of modifying the Profile directly
 		 * 
 		 * @param CommandLine the new command line
 		 */
+		UE_DEPRECATED(5.8, "access the appropriate field from a BuildCookRun instance, or use ICmdLineParametersExtension")
 		UE_API void SetCommandLine( const FString& CommandLine );
 
 		/**
@@ -359,10 +367,6 @@ namespace ProjectLauncher
 		TSharedRef<FLaunchExtension> Extension;
 		FLaunchProfileTreeData& OwnerTreeData;
 		FPropertyChangedDelegate PropertyChangedDelegate;
-
-		friend class FLaunchProfileTreeNode;
-		UE_API void MakeCommandLineSubmenu( FMenuBuilder& MenuBuilder );
-
 	};
 
 
diff --git a/Engine/Plugins/Developer/ProjectLauncher/Source/ProjectLauncher/Public/Model/ProjectLauncherModel.h b/Engine/Plugins/Developer/ProjectLauncher/Source/ProjectLauncher/Public/Model/ProjectLauncherModel.h
--- a/Engine/Plugins/Developer/ProjectLauncher/Source/ProjectLauncher/Public/Model/ProjectLauncherModel.h
+++ b/Engine/Plugins/Developer/ProjectLauncher/Source/ProjectLauncher/Public/Model/ProjectLauncherModel.h
@@ -120,6 +120,14 @@ namespace ProjectLauncher
 	DECLARE_MULTICAST_DELEGATE(FOnSelectedPlatformChanged)
 
 
+	/* common command line types */
+	namespace CmdLineType
+	{
+		extern const TCHAR* Launch;
+		extern const TCHAR* UAT;
+	}
+
+
 	/*
 	 * This class contains a selection of helper functions and the current profile state
 	 */
diff --git a/Engine/Plugins/Developer/ProjectLauncher/Source/ProjectLauncher/Public/ProfileTree/LaunchProfileTreeData.h b/Engine/Plugins/Developer/ProjectLauncher/Source/ProjectLauncher/Public/ProfileTree/LaunchProfileTreeData.h
--- a/Engine/Plugins/Developer/ProjectLauncher/Source/ProjectLauncher/Public/ProfileTree/LaunchProfileTreeData.h
+++ b/Engine/Plugins/Developer/ProjectLauncher/Source/ProjectLauncher/Public/ProfileTree/LaunchProfileTreeData.h
@@ -124,7 +124,7 @@ namespace ProjectLauncher
 		UE_API FLaunchProfileTreeNode& AddString( FText InName, FStringCallbacks&& StringCallbacks );
 		UE_API FLaunchProfileTreeNode& AddDirectoryString(FText InName, FStringCallbacks&& StringCallbacks);
 		UE_API FLaunchProfileTreeNode& AddFileString(FText InName, FStringCallbacks&& StringCallbacks, FString FileFilter = TEXT("All files (*.*)|*.*"));
-		UE_API FLaunchProfileTreeNode& AddCommandLineString( FText InName, FStringCallbacks&& StringCallbacks );
+		UE_API FLaunchProfileTreeNode& AddCommandLineString( FText InName, FStringCallbacks&& StringCallbacks, const TCHAR* InType = nullptr ); // defaults to "Launch" cmdline extensions
 
 		/* Generic value callbacks. */
 		template<typename T>
