ProjectLauncher2 - disable classic view and grey out the tickbox of the default BuildCookRun step when profile is first created

#Synced-CL 49728215

diff --git a/Engine/Plugins/Developer/ProjectLauncher/Source/ProjectLauncher/Private/Model/ProjectLauncherModel.cpp b/Engine/Plugins/Developer/ProjectLauncher/Source/ProjectLauncher/Private/Model/ProjectLauncherModel.cpp
--- a/Engine/Plugins/Developer/ProjectLauncher/Source/ProjectLauncher/Private/Model/ProjectLauncherModel.cpp
+++ b/Engine/Plugins/Developer/ProjectLauncher/Source/ProjectLauncher/Private/Model/ProjectLauncherModel.cpp
@@ -28,7 +28,7 @@
 namespace ProjectLauncher
 {
 	const FText ProjectLauncherText = LOCTEXT("ProjectLauncherTitle", "Project Launcher");
-	bool bUseClassicView = true;
+	bool bUseClassicView = false;
 
 	namespace CmdLineType
 	{
diff --git a/Engine/Plugins/Developer/ProjectLauncher/Source/ProjectLauncher/Private/Widgets/Profiles/SCustomLaunchCustomProfileEditor.cpp b/Engine/Plugins/Developer/ProjectLauncher/Source/ProjectLauncher/Private/Widgets/Profiles/SCustomLaunchCustomProfileEditor.cpp
--- a/Engine/Plugins/Developer/ProjectLauncher/Source/ProjectLauncher/Private/Widgets/Profiles/SCustomLaunchCustomProfileEditor.cpp
+++ b/Engine/Plugins/Developer/ProjectLauncher/Source/ProjectLauncher/Private/Widgets/Profiles/SCustomLaunchCustomProfileEditor.cpp
@@ -221,6 +221,7 @@ public:
 			.AutoWidth()
 			[
 				SNew(SCheckBox)
+				.IsEnabled_Lambda([this]() { return TreeNode->GetTreeData()->Profile->GetUATCommands().Num() > 1; })
 				.IsChecked_Lambda(GetUATCommandEnabledCheckState)
 				.OnCheckStateChanged_Lambda(SetUATCommandEnabledCheckState)
 			];
