Project Launcher UI tweaks
- add icons for UAT command heading & menus
- replace header's Extensions menu with an overflow menu containing the full context menu

#rnx
#Synced-CL 49676362

diff --git a/Engine/Plugins/Developer/ProjectLauncher/Source/ProjectLauncher/Private/Extension/AutomatedTestLaunchExtension.cpp b/Engine/Plugins/Developer/ProjectLauncher/Source/ProjectLauncher/Private/Extension/AutomatedTestLaunchExtension.cpp
--- a/Engine/Plugins/Developer/ProjectLauncher/Source/ProjectLauncher/Private/Extension/AutomatedTestLaunchExtension.cpp
+++ b/Engine/Plugins/Developer/ProjectLauncher/Source/ProjectLauncher/Private/Extension/AutomatedTestLaunchExtension.cpp
@@ -1,6 +1,7 @@
 // Copyright Epic Games, Inc. All Rights Reserved.
 
 #include "Extension/AutomatedTestLaunchExtension.h"
+#include "Styling/ProjectLauncherStyle.h"
 
 #define LOCTEXT_NAMESPACE "FAutomatedTestLaunchExtensionInstance"
 
@@ -106,6 +107,12 @@ namespace ProjectLauncher
 	{
 		MenuEntry = FExtensionsMenuEntry::AutomatedTests;
 	}
+
+	FSlateIcon FAutomatedTestLaunchExtension::GetIcon() const
+	{
+		return FSlateIcon( FProjectLauncherStyle::GetStyleSetName(), "Icons.Task.TestAutomation");
+	}
+
 }
 
 #undef LOCTEXT_NAMESPACE
diff --git a/Engine/Plugins/Developer/ProjectLauncher/Source/ProjectLauncher/Private/Extension/BuildCookRunLaunchExtension.cpp b/Engine/Plugins/Developer/ProjectLauncher/Source/ProjectLauncher/Private/Extension/BuildCookRunLaunchExtension.cpp
--- a/Engine/Plugins/Developer/ProjectLauncher/Source/ProjectLauncher/Private/Extension/BuildCookRunLaunchExtension.cpp
+++ b/Engine/Plugins/Developer/ProjectLauncher/Source/ProjectLauncher/Private/Extension/BuildCookRunLaunchExtension.cpp
@@ -2,6 +2,7 @@
 
 #include "Extension/BuildCookRunLaunchExtension.h"
 #include "Model/ProjectLauncherModel.h"
+#include "Styling/ProjectLauncherStyle.h"
 
 #define LOCTEXT_NAMESPACE "FBuildCookRunLaunchExtensionInstance"
 
@@ -181,7 +182,7 @@ namespace ProjectLauncher
 		InMenuBuilder.AddMenuEntry(
 			GetExtension()->GetDisplayName(),
 			GetToolTipText(),
-			FSlateIcon(),
+			GetExtension()->GetIcon(),
 			FUIAction(
 				FExecuteAction::CreateLambda(ToggleEnabled),
 				FCanExecuteAction::CreateLambda(CanToggle),
diff --git a/Engine/Plugins/Developer/ProjectLauncher/Source/ProjectLauncher/Private/Extension/LaunchExtension.cpp b/Engine/Plugins/Developer/ProjectLauncher/Source/ProjectLauncher/Private/Extension/LaunchExtension.cpp
--- a/Engine/Plugins/Developer/ProjectLauncher/Source/ProjectLauncher/Private/Extension/LaunchExtension.cpp
+++ b/Engine/Plugins/Developer/ProjectLauncher/Source/ProjectLauncher/Private/Extension/LaunchExtension.cpp
@@ -2,6 +2,7 @@
 
 #include "Extension/LaunchExtension.h"
 #include "Extension/CmdLineParametersExtension.h"
+#include "Styling/ProjectLauncherStyle.h"
 #include "Model/ProjectLauncherModel.h"
 #include "Utils/LauncherCmdLineUtils.h"
 #include "Containers/Array.h"
@@ -375,6 +376,10 @@ namespace ProjectLauncher
 		MenuEntry = FExtensionsMenuEntry::Default;
 	}
 
+	FSlateIcon FLaunchExtension::GetIcon() const
+	{
+		return FSlateIcon( FProjectLauncherStyle::GetStyleSetName(), "Icons.Plugins" );
+	}
 
 
 
diff --git a/Engine/Plugins/Developer/ProjectLauncher/Source/ProjectLauncher/Private/Widgets/Profiles/SCustomLaunchCustomProfileEditor.cpp b/Engine/Plugins/Developer/ProjectLauncher/Source/ProjectLauncher/Private/Widgets/Profiles/SCustomLaunchCustomProfileEditor.cpp
--- a/Engine/Plugins/Developer/ProjectLauncher/Source/ProjectLauncher/Private/Widgets/Profiles/SCustomLaunchCustomProfileEditor.cpp
+++ b/Engine/Plugins/Developer/ProjectLauncher/Source/ProjectLauncher/Private/Widgets/Profiles/SCustomLaunchCustomProfileEditor.cpp
@@ -176,6 +176,10 @@ public:
 			return FReply::Handled();
 		};
 
+		auto GetUATCommandBrush = [this]()
+		{
+			return ProfileEditor->GetIconForUATCommand(TreeNode->UATCommand.ToSharedRef()).GetIcon();
+		};
 
 
 		TSharedPtr<SHorizontalBox> HeaderBar;
@@ -243,6 +247,17 @@ public:
 
 		if (bIsPrimaryHeading)
 		{
+			// icon
+			HeaderBar->AddSlot()
+			.VAlign(VAlign_Center)
+			.Padding(4, 0, 0, 0)
+			.AutoWidth()
+			[
+				SNew(SImage)
+				.Image(GetUATCommandBrush())
+				.ColorAndOpacity(FSlateColor::UseSubduedForeground())
+			];
+
 			// inline editable name field for UAT commands etc
 			TSharedPtr<SInlineEditableTextBlock> RenameBox;
 			HeaderBar->AddSlot()
@@ -271,10 +286,8 @@ public:
 			];
 		}
 
-		if (bIsPrimaryHeading && TreeNode->UATCommand->AsBuildCookRun().IsValid())
+		if (bIsPrimaryHeading)
 		{
-			ILauncherProfileBuildCookRunRef BuildCookRun = TreeNode->UATCommand->AsBuildCookRun().ToSharedRef();
-
 			// spacer to push everthing else to the right
 			HeaderBar->AddSlot()
 			.FillWidth(1)
@@ -291,13 +304,12 @@ public:
 			[
 				SNew(SComboButton)
 				.ButtonStyle(FAppStyle::Get(), "SimpleButton")
-				.ToolTipText(LOCTEXT("BCRExtensionMenuTip", "Enable an extension for this Build, Cook & Run"))
-				.OnGetMenuContent( ProfileEditor, &SCustomLaunchCustomProfileEditor::MakeBuildCookRunExtensionsSubMenuWidget, BuildCookRun)
+				.HasDownArrow(false)
+				.OnGetMenuContent( ProfileEditor, &SCustomLaunchCustomProfileEditor::MakeUATCommandSubMenuWidget, TreeNode->UATCommand.ToSharedRef())
 				.ButtonContent()
 				[
-					SNew(SImage)
-					.Image(FProjectLauncherStyle::Get().GetBrush("Icons.Plugins"))
-					.ColorAndOpacity(FSlateColor::UseSubduedForeground())
+					SNew(STextBlock)
+					.Text(LOCTEXT("...", "..."))
 				]
 			];
 		}
@@ -1104,7 +1116,7 @@ void SCustomLaunchCustomProfileEditor::MakeExtensionMenu( FMenuBuilder& MenuBuil
 		MenuBuilder.AddMenuEntry(
 			Extension->GetDisplayName(),
 			GetToolTipText(),
-			FSlateIcon(),
+			Extension->GetIcon(),
 			FUIAction(
 				FExecuteAction::CreateLambda(OnToggleExtension),
 				FCanExecuteAction::CreateLambda(CanToggleExtension),
@@ -1151,7 +1163,7 @@ TSharedRef<SWidget> SCustomLaunchCustomProfileEditor::MakeCommandsMenu()
 		MenuBuilder.AddMenuEntry(
 			LOCTEXT("NewBuildCookRunCmdLabel", "Add New Build, Cook & Run"),
 			FText::GetEmpty(),
-			FSlateIcon(),
+			FSlateIcon(FProjectLauncherStyle::GetStyleSetName(), "Icons.Task.Launch"),
 			FUIAction(
 				FExecuteAction::CreateLambda( [this, TreeData]()
 				{
@@ -1237,19 +1249,14 @@ void SCustomLaunchCustomProfileEditor::MakeUATCommandSubMenu( FMenuBuilder& Menu
 				if (FPlatformMisc::MessageBoxExt(EAppMsgType::YesNo, TEXT("Confirm task deletion?"), TEXT("Delete task")) == EAppReturnType::Yes)
 				{
 					// see if there is an extension associated with this UAT command
-					TSharedPtr<ProjectLauncher::FLaunchExtensionInstance>* ExtensionInstancePtr = TreeData->ExtensionInstances.FindByPredicate( [UATCommand]( TSharedPtr<ProjectLauncher::FLaunchExtensionInstance> ExtensionInstance )
-					{
-						return ExtensionInstance->ManagesUATCommand(UATCommand);
-					});
-					TSharedPtr<ProjectLauncher::FLaunchExtension> Extension = (ExtensionInstancePtr != nullptr) ? (*ExtensionInstancePtr)->GetExtension().ToSharedPtr() : nullptr;
-
+					TSharedPtr<ProjectLauncher::FLaunchExtensionInstance> ExtensionInstance = GetExtensionForUATCommand(UATCommand);
 
 					TreeData->Profile->RemoveUATCommand(UATCommand->GetInternalName()); 
 
-					if (Extension.IsValid())
+					if (ExtensionInstance.IsValid())
 					{
 						// defer the delete until the next tick
-						PendingExtensionDeletes.AddUnique(Extension.ToSharedRef());
+						PendingExtensionDeletes.AddUnique(ExtensionInstance->GetExtension());
 					}
 					else
 					{
@@ -1285,7 +1292,7 @@ void SCustomLaunchCustomProfileEditor::MakeBuildCookRunExtensionsSubMenu( FMenuB
 }
 
 
-TSharedRef<SWidget> SCustomLaunchCustomProfileEditor::MakeBuildCookRunExtensionsSubMenuWidget(ILauncherProfileBuildCookRunRef BuildCookRun) const
+TSharedRef<SWidget> SCustomLaunchCustomProfileEditor::MakeUATCommandSubMenuWidget(ILauncherProfileUATCommandRef UATCommand)
 {
 	const bool bShouldCloseWindowAfterMenuSelection = true;
 	const bool bCloseSelfOnly = true;
@@ -1293,11 +1300,48 @@ TSharedRef<SWidget> SCustomLaunchCustomProfileEditor::MakeBuildCookRunExtensions
 	const bool bInOpenSubMenuOnClick = false;
 	FMenuBuilder MenuBuilder(bShouldCloseWindowAfterMenuSelection, nullptr, nullptr, bCloseSelfOnly, &FCoreStyle::Get(), bSearchable );
 
-	MakeBuildCookRunExtensionsSubMenu(MenuBuilder, BuildCookRun);
+	MakeUATCommandSubMenu(MenuBuilder, UATCommand);
 
 	return MenuBuilder.MakeWidget();
 }
 
+TSharedPtr<ProjectLauncher::FLaunchExtensionInstance> SCustomLaunchCustomProfileEditor::GetExtensionForUATCommand(ILauncherProfileUATCommandRef UATCommand) const
+{
+	ProjectLauncher::FLaunchProfileTreeDataRef TreeData = TreeBuilder->GetProfileTree();
+
+	TSharedPtr<ProjectLauncher::FLaunchExtensionInstance>* ExtensionInstancePtr = TreeData->ExtensionInstances.FindByPredicate( [UATCommand]( TSharedPtr<ProjectLauncher::FLaunchExtensionInstance> ExtensionInstance )
+	{
+		return ExtensionInstance->ManagesUATCommand(UATCommand);
+	});
+	
+	if ( ExtensionInstancePtr != nullptr)
+	{
+		return *ExtensionInstancePtr;
+	}
+
+	return nullptr;
+}
+
+
+FSlateIcon SCustomLaunchCustomProfileEditor::GetIconForUATCommand(ILauncherProfileUATCommandPtr UATCommand) const
+{
+	if (UATCommand.IsValid())
+	{
+		TSharedPtr<ProjectLauncher::FLaunchExtensionInstance> ExtensionInstance = GetExtensionForUATCommand(UATCommand.ToSharedRef());
+		if (ExtensionInstance.IsValid())
+		{
+			return ExtensionInstance->GetExtension()->GetIcon();
+		}
+
+		if (UATCommand->AsBuildCookRun() != nullptr)
+		{
+			return FSlateIcon(FProjectLauncherStyle::GetStyleSetName(), "Icons.Task.Launch");
+		}
+	}
+
+	return FSlateIcon();
+}
+
 
 void SCustomLaunchCustomProfileEditor::Tick(const FGeometry& AllottedGeometry, const double InCurrentTime, const float InDeltaTime)
 {
diff --git a/Engine/Plugins/Developer/ProjectLauncher/Source/ProjectLauncher/Private/Widgets/Profiles/SCustomLaunchCustomProfileEditor.h b/Engine/Plugins/Developer/ProjectLauncher/Source/ProjectLauncher/Private/Widgets/Profiles/SCustomLaunchCustomProfileEditor.h
--- a/Engine/Plugins/Developer/ProjectLauncher/Source/ProjectLauncher/Private/Widgets/Profiles/SCustomLaunchCustomProfileEditor.h
+++ b/Engine/Plugins/Developer/ProjectLauncher/Source/ProjectLauncher/Private/Widgets/Profiles/SCustomLaunchCustomProfileEditor.h
@@ -10,7 +10,7 @@
 #include "Widgets/Views/STableRow.h"
 #include "ProfileTree/ILaunchProfileTreeBuilder.h"
 
-
+struct FSlateIcon;
 class SWidget;
 class ITableRow;
 class STableViewBase;
@@ -61,7 +61,10 @@ private:
 	void MakeUATCommandSubMenu( FMenuBuilder& MenuBuilder, ILauncherProfileUATCommandRef UATCommand );
 	void MakeBuildCookRunExtensionsSubMenu( FMenuBuilder& MenuBuilder, ILauncherProfileBuildCookRunRef BuildCookRun ) const;
 	void MakeExtensionMenu( FMenuBuilder& MenuBuilder, TSharedRef<ProjectLauncher::FLaunchExtension> Extension, bool bIncludeBuildCookRunExtensions );
-	TSharedRef<SWidget> MakeBuildCookRunExtensionsSubMenuWidget(ILauncherProfileBuildCookRunRef BuildCookRun) const;
+	TSharedRef<SWidget> MakeUATCommandSubMenuWidget(ILauncherProfileUATCommandRef UATCommand);
+
+	TSharedPtr<ProjectLauncher::FLaunchExtensionInstance> GetExtensionForUATCommand(ILauncherProfileUATCommandRef UATCommand) const;
+	FSlateIcon GetIconForUATCommand(ILauncherProfileUATCommandPtr UATCommand) const;
 
 	ILauncherProfilePtr CurrentProfile;
 	TSharedPtr<ProjectLauncher::FModel> Model;
diff --git a/Engine/Plugins/Developer/ProjectLauncher/Source/ProjectLauncher/Public/Extension/AutomatedTestLaunchExtension.h b/Engine/Plugins/Developer/ProjectLauncher/Source/ProjectLauncher/Public/Extension/AutomatedTestLaunchExtension.h
--- a/Engine/Plugins/Developer/ProjectLauncher/Source/ProjectLauncher/Public/Extension/AutomatedTestLaunchExtension.h
+++ b/Engine/Plugins/Developer/ProjectLauncher/Source/ProjectLauncher/Public/Extension/AutomatedTestLaunchExtension.h
@@ -78,6 +78,8 @@ namespace ProjectLauncher
 	{
 	public:
 		UE_API virtual void GetExtensionsMenuEntry( FExtensionsMenuEntry& MenuEntry ) const override;
+		UE_API virtual FSlateIcon GetIcon() const override;
+
 		virtual class FAutomatedTestLaunchExtension* AsAutomatedTestLaunchExtension() override { return this; }
 	};
 };
diff --git a/Engine/Plugins/Developer/ProjectLauncher/Source/ProjectLauncher/Public/Extension/LaunchExtension.h b/Engine/Plugins/Developer/ProjectLauncher/Source/ProjectLauncher/Public/Extension/LaunchExtension.h
--- a/Engine/Plugins/Developer/ProjectLauncher/Source/ProjectLauncher/Public/Extension/LaunchExtension.h
+++ b/Engine/Plugins/Developer/ProjectLauncher/Source/ProjectLauncher/Public/Extension/LaunchExtension.h
@@ -397,6 +397,11 @@ namespace ProjectLauncher
 		 */
 		virtual FText GetDisplayName() const = 0;
 
+		/**
+		 * Returns the icon to use for this extension in menus etc
+		 */
+		UE_API virtual FSlateIcon GetIcon() const;
+
 		/**
 		 * Whether this extension is always created. the user cannot remove it
 		 */
