automatically choose a build target if there is no default specified in the project
- should make 'client only' projects launch correctly without having to specify a build target first

#rnx
#Synced-CL 49676362

diff --git a/Engine/Plugins/Developer/ProjectLauncher/Source/ProjectLauncher/Private/Model/ProjectLauncherModel.cpp b/Engine/Plugins/Developer/ProjectLauncher/Source/ProjectLauncher/Private/Model/ProjectLauncherModel.cpp
--- a/Engine/Plugins/Developer/ProjectLauncher/Source/ProjectLauncher/Private/Model/ProjectLauncherModel.cpp
+++ b/Engine/Plugins/Developer/ProjectLauncher/Source/ProjectLauncher/Private/Model/ProjectLauncherModel.cpp
@@ -115,6 +115,15 @@ namespace ProjectLauncher
 			ILauncherProfilePtr PreviousProfile = SelectedProfile;
 			SelectedProfile = NewProfile;
 
+			// try to make sure there is a valid build target
+			if (SelectedProfile.IsValid())
+			{
+				for (const ILauncherProfileBuildCookRunRef& BuildCookRun : SelectedProfile->GetBuildCookRunCommands())
+				{
+					EnsureValidBuildTarget(SelectedProfile.ToSharedRef(), BuildCookRun);
+				}
+			}
+
 			ProfileSelectedDelegate.Broadcast(SelectedProfile, PreviousProfile);
 		}
 	}
@@ -554,6 +563,7 @@ PRAGMA_ENABLE_DEPRECATION_WARNINGS
 			BuildCookRun->AddCookedPlatform(FPlatformProperties::IniPlatformName());
 		}
 
+		EnsureValidBuildTarget(Profile, BuildCookRun);
 	}
 
 	ILauncherProfileRef FModel::CreateBasicLaunchProfile()
@@ -684,6 +694,23 @@ PRAGMA_ENABLE_DEPRECATION_WARNINGS
 	}
 
 
+	void FModel::EnsureValidBuildTarget(const ILauncherProfileRef& Profile, const ILauncherProfileBuildCookRunRef& BuildCookRun)
+	{
+		if (Profile->HasProjectSpecified())
+		{
+			const FProjectSettings& ProjectSettings = GetProjectSettings(Profile);
+			if (!ProjectSettings.FallbackDefaultBuildTarget.IsEmpty() && (BuildCookRun->GetBuildTargets().Num() == 0 || BuildCookRun->GetBuildTargets().Contains(FString())))
+			{
+				BuildCookRun->ClearBuildTargets();
+				BuildCookRun->AddBuildTarget(ProjectSettings.FallbackDefaultBuildTarget);
+				BuildCookRun->SetBuildTargetSpecified(true);
+
+				UpdateCookedPlatformsFromBuildTarget(Profile, BuildCookRun);
+			}
+		}
+	}
+
+
 	ILauncherProfilePtr FModel::CloneCustomProfile(const ILauncherProfileRef& Profile)
 	{
 		// capture the current profile
@@ -847,9 +874,33 @@ PRAGMA_ENABLE_DEPRECATION_WARNINGS
 			}
 			// multiple (or no) game targets... would need to read the BuildTarget from /Script/BuildSettings.BuildSettings... also reading the platform-specific ini
 			// ... fixme ...
-		}
 
+			// find a fallback build target
+			// pick the first Game build target... if there are no Game build targets available, default to a Client build target if we can
+			if (OutResult.DefaultBuildTargetName.IsEmpty())
+			{
+				if (GameBuildTargets.Num() > 0)
+				{
+					OutResult.FallbackDefaultBuildTarget = GameBuildTargets[0].Name;
+				}
+				else if (!OutResult.DefaultClientTargetName.IsEmpty())
+				{
+					OutResult.FallbackDefaultBuildTarget = OutResult.DefaultClientTargetName;
+				}
+				else
+				{
+					const TArray<FTargetInfo> ClientBuildTargets = BuildTargets.FilterByPredicate([](const FTargetInfo& TargetInfo)
+					{
+						return TargetInfo.Type == EBuildTargetType::Client;
+					});
+					if (ClientBuildTargets.Num() > 0)
+					{
+						OutResult.FallbackDefaultBuildTarget = ClientBuildTargets[0].Name;
+					}
+				}
+			}
 
+		}
 		
 		// read zen settings
 #if UE_WITH_ZEN
@@ -860,6 +911,7 @@ PRAGMA_ENABLE_DEPRECATION_WARNINGS
 		const TCHAR* ZenAutoLaunchSettings = TEXT("Zen.AutoLaunch");
 		InConfig.GetBool(ZenAutoLaunchSettings, TEXT("AllowRemoteNetworkService"), OutResult.bAllowRemoteNetworkService, GEngineIni);
 #endif
+
 	}
 
 	const FProjectSettings FModel::GetProjectSettings( const FString& InProjectPath )
diff --git a/Engine/Plugins/Developer/ProjectLauncher/Source/ProjectLauncher/Private/ProfileTree/GenericProfileTreeBuilder.cpp b/Engine/Plugins/Developer/ProjectLauncher/Source/ProjectLauncher/Private/ProfileTree/GenericProfileTreeBuilder.cpp
--- a/Engine/Plugins/Developer/ProjectLauncher/Source/ProjectLauncher/Private/ProfileTree/GenericProfileTreeBuilder.cpp
+++ b/Engine/Plugins/Developer/ProjectLauncher/Source/ProjectLauncher/Private/ProfileTree/GenericProfileTreeBuilder.cpp
@@ -1050,9 +1050,11 @@ namespace ProjectLauncher
 		Profile->SetProjectPath(ProjectPath);
 		Profile->SetProjectSpecified(!ProjectPath.IsEmpty());
 
-		if (ProfileType != EProfileType::Basic )
+		for (const ILauncherProfileBuildCookRunRef& BuildCookRun : Profile->GetBuildCookRunCommands())
 		{
-			for (const ILauncherProfileBuildCookRunRef& BuildCookRun : Profile->GetBuildCookRunCommands())
+			Model->EnsureValidBuildTarget(Profile, BuildCookRun);
+
+			if (ProfileType != EProfileType::Basic )
 			{
 				Model->UpdateCookedPlatformsFromBuildTarget(Profile, BuildCookRun);
 			}
diff --git a/Engine/Plugins/Developer/ProjectLauncher/Source/ProjectLauncher/Private/Widgets/Shared/SCustomLaunchBuildTargetCombo.cpp b/Engine/Plugins/Developer/ProjectLauncher/Source/ProjectLauncher/Private/Widgets/Shared/SCustomLaunchBuildTargetCombo.cpp
--- a/Engine/Plugins/Developer/ProjectLauncher/Source/ProjectLauncher/Private/Widgets/Shared/SCustomLaunchBuildTargetCombo.cpp
+++ b/Engine/Plugins/Developer/ProjectLauncher/Source/ProjectLauncher/Private/Widgets/Shared/SCustomLaunchBuildTargetCombo.cpp
@@ -49,7 +49,7 @@ TSharedRef<SWidget> SCustomLaunchBuildTargetCombo::MakeBuildTargetSelectionWidge
 		EUserInterfaceActionType UIType = GetAdvancedPlatformsOption.Get() ? EUserInterfaceActionType::Check : EUserInterfaceActionType::RadioButton;
 
 		FString DefaultBuildTarget = GetDefaultBuildTargetName();
-		FText DefaultBuildTargetName = DefaultBuildTarget.IsEmpty() ? LOCTEXT("DefaultBuildTargetName", "Project Default") : FText::Format(LOCTEXT("CurProjectDefaultBuildTargetName", "{0} (Project Default)"), FText::FromString(DefaultBuildTarget));
+		FText DefaultBuildTargetName = DefaultBuildTarget.IsEmpty() ? LOCTEXT("DefaultBuildTargetName", "Not Specified") : FText::Format(LOCTEXT("CurProjectDefaultBuildTargetName", "{0} (Project Default)"), FText::FromString(DefaultBuildTarget));
 		MenuBuilder.AddMenuEntry(
 			DefaultBuildTargetName, 
 			LOCTEXT("DefaultBuildTargetActionHint", "Use the project default build target."), 
@@ -132,7 +132,7 @@ FText SCustomLaunchBuildTargetCombo::GetBuildTargetName() const
 			FString DefaultBuildTarget = GetDefaultBuildTargetName();
 			if (DefaultBuildTarget.IsEmpty())
 			{
-				return LOCTEXT("DefaultBuildTargetName", "Project Default");
+				return LOCTEXT("DefaultBuildTargetName", "Not Specified");
 			}
 			else
 			{
diff --git a/Engine/Plugins/Developer/ProjectLauncher/Source/ProjectLauncher/Public/Model/ProjectLauncherModel.h b/Engine/Plugins/Developer/ProjectLauncher/Source/ProjectLauncher/Public/Model/ProjectLauncherModel.h
--- a/Engine/Plugins/Developer/ProjectLauncher/Source/ProjectLauncher/Public/Model/ProjectLauncherModel.h
+++ b/Engine/Plugins/Developer/ProjectLauncher/Source/ProjectLauncher/Public/Model/ProjectLauncherModel.h
@@ -105,6 +105,9 @@ namespace ProjectLauncher
 		FString DefaultClientTargetName;
 		FString DefaultEditorTargetName;
 
+		// fallback "default" build target to use if none is specified
+		FString FallbackDefaultBuildTarget;
+
 		// cached project config ini
 		FConfigCacheIni* Config = nullptr;
 	};
@@ -269,6 +272,9 @@ namespace ProjectLauncher
 		/* Update cooked platforms based on build target */
 		UE_API void UpdateCookedPlatformsFromBuildTarget(const ILauncherProfileRef& Profile, const ILauncherProfileBuildCookRunRef& BuildCookRun);
 
+		/* Attempts to set a valid build target if one has not been specified */
+		UE_API void EnsureValidBuildTarget(const ILauncherProfileRef& Profile, const ILauncherProfileBuildCookRunRef& BuildCookRun);
+
 		/* Utility functions for resolving platform names */
 		static UE_API FString GetVanillaPlatformName(const FString& PlatformName);
 		static UE_API FString GetPlatformNameWithFlavor(const FString& PlatformName);
