Use slate-based confirmation message dialog to avoid window focus issues

#rnx

#Synced-CL 49921272

diff --git a/Engine/Plugins/Developer/ProjectLauncher/Source/ProjectLauncher/Private/Model/ProjectLauncherModel.cpp b/Engine/Plugins/Developer/ProjectLauncher/Source/ProjectLauncher/Private/Model/ProjectLauncherModel.cpp
--- a/Engine/Plugins/Developer/ProjectLauncher/Source/ProjectLauncher/Private/Model/ProjectLauncherModel.cpp
+++ b/Engine/Plugins/Developer/ProjectLauncher/Source/ProjectLauncher/Private/Model/ProjectLauncherModel.cpp
@@ -13,6 +13,7 @@
 #include "ITargetDeviceServicesModule.h"
 #include "Serialization/JsonReader.h"
 #include "Serialization/JsonSerializer.h"
+#include "Dialog/SMessageDialog.h"
 #include "Modules/ModuleManager.h"
 #include "Settings/ProjectPackagingSettings.h"
 #include "Experimental/ZenServerInterface.h"
@@ -1397,6 +1398,27 @@ PRAGMA_ENABLE_DEPRECATION_WARNINGS
 
 		return MsgTextBuilder.ToText();
 	}
+
+
+
+	bool GetUserConfirmation( const FText& Message, const FText& Title, bool bDefaultValue )
+	{
+		TSharedRef<SMessageDialog> Dialog = SNew(SMessageDialog)
+			.Title(Title)
+			.Message(Message)
+			.Buttons({SMessageDialog::FButton(LOCTEXT("YES","Yes")), SMessageDialog::FButton(LOCTEXT("NO","No"))})
+		;
+
+		int32 DefaultButtonIndex = bDefaultValue ? 0 : 1;
+		if (Dialog->ShowModal(DefaultButtonIndex) != 0)
+		{
+			return false;
+		}
+
+		return true;
+	}
+
+
 }
 
 
diff --git a/Engine/Plugins/Developer/ProjectLauncher/Source/ProjectLauncher/Private/ProfileTree/GenericProfileTreeBuilder.cpp b/Engine/Plugins/Developer/ProjectLauncher/Source/ProjectLauncher/Private/ProfileTree/GenericProfileTreeBuilder.cpp
--- a/Engine/Plugins/Developer/ProjectLauncher/Source/ProjectLauncher/Private/ProfileTree/GenericProfileTreeBuilder.cpp
+++ b/Engine/Plugins/Developer/ProjectLauncher/Source/ProjectLauncher/Private/ProfileTree/GenericProfileTreeBuilder.cpp
@@ -11,7 +11,6 @@
 #include "Widgets/Text/STextBlock.h"
 #include "Widgets/Layout/SBox.h"
 #include "Widgets/SBoxPanel.h"
-#include "Dialog/SMessageDialog.h"
 #include "SResizeBox.h"
 #include "SSearchableComboBox.h"
 #include "PlatformInfo.h"
@@ -1636,12 +1635,7 @@ namespace ProjectLauncher
 			if (BuildCookRun->GetCookedPlatforms().Num() > 1 || BuildCookRun->GetBuildTargets().Num() > 1)
 			{
 				FText MessageText = LOCTEXT("ClearAdvancedPlatformTargets", "You have multiple platform/targets selected - are you sure you want to go back to single platform/target??");
-				TSharedRef<SMessageDialog> Dialog = SNew(SMessageDialog)
-					.Title(ProjectLauncherText)
-					.Message(MessageText)
-					.Buttons({SMessageDialog::FButton(LOCTEXT("YES","Yes")), SMessageDialog::FButton(LOCTEXT("NO","No"))})
-				;
-				if (Dialog->ShowModal() != 0)
+				if (!ProjectLauncher::GetUserConfirmation(MessageText, ProjectLauncher::ProjectLauncherText, false))
 				{
 					return;
 				}
diff --git a/Engine/Plugins/Developer/ProjectLauncher/Source/ProjectLauncher/Private/Widgets/Profiles/SCustomLaunchCustomProfileEditor.cpp b/Engine/Plugins/Developer/ProjectLauncher/Source/ProjectLauncher/Private/Widgets/Profiles/SCustomLaunchCustomProfileEditor.cpp
--- a/Engine/Plugins/Developer/ProjectLauncher/Source/ProjectLauncher/Private/Widgets/Profiles/SCustomLaunchCustomProfileEditor.cpp
+++ b/Engine/Plugins/Developer/ProjectLauncher/Source/ProjectLauncher/Private/Widgets/Profiles/SCustomLaunchCustomProfileEditor.cpp
@@ -1245,8 +1245,7 @@ void SCustomLaunchCustomProfileEditor::MakeUATCommandSubMenu( FMenuBuilder& Menu
 		FUIAction(
 			FExecuteAction::CreateLambda( [this, TreeData, UATCommand]() 
 			{
-				// fixme: make this better
-				if (FPlatformMisc::MessageBoxExt(EAppMsgType::YesNo, TEXT("Confirm task deletion?"), TEXT("Delete task")) == EAppReturnType::Yes)
+				if (ProjectLauncher::GetUserConfirmation(LOCTEXT("DeleteTaskMsg","Confirm task deletion?"), LOCTEXT("DeleteTaskTitle","Delete task"), false))
 				{
 					// see if there is an extension associated with this UAT command
 					TSharedPtr<ProjectLauncher::FLaunchExtensionInstance> ExtensionInstance = GetExtensionForUATCommand(UATCommand);
diff --git a/Engine/Plugins/Developer/ProjectLauncher/Source/ProjectLauncher/Private/Widgets/Profiles/SCustomLaunchProfilesPanel.cpp b/Engine/Plugins/Developer/ProjectLauncher/Source/ProjectLauncher/Private/Widgets/Profiles/SCustomLaunchProfilesPanel.cpp
--- a/Engine/Plugins/Developer/ProjectLauncher/Source/ProjectLauncher/Private/Widgets/Profiles/SCustomLaunchProfilesPanel.cpp
+++ b/Engine/Plugins/Developer/ProjectLauncher/Source/ProjectLauncher/Private/Widgets/Profiles/SCustomLaunchProfilesPanel.cpp
@@ -611,7 +611,7 @@ void SCustomLaunchProfilesPanel::OnProfileDelete(const ILauncherProfilePtr& Prof
 	if (Profile.IsValid() && !Model->IsBasicLaunchProfile(Profile))
 	{
 		FText MessageText = LOCTEXT("ProfileDeleteConfirm", "Are you sure you want to delete this profile?");
-		if (FPlatformMisc::MessageBoxExt(EAppMsgType::YesNo, *MessageText.ToString(), *Profile->GetName()) == EAppReturnType::Yes)
+		if (ProjectLauncher::GetUserConfirmation(MessageText, FText::FromString(*Profile->GetName()), false))
 		{
 			Model->GetProfileManager()->RemoveProfile(Profile.ToSharedRef());
 			Model->SelectProfile(Model->GetBasicLaunchProfile());
diff --git a/Engine/Plugins/Developer/ProjectLauncher/Source/ProjectLauncher/Public/Model/ProjectLauncherModel.h b/Engine/Plugins/Developer/ProjectLauncher/Source/ProjectLauncher/Public/Model/ProjectLauncherModel.h
--- a/Engine/Plugins/Developer/ProjectLauncher/Source/ProjectLauncher/Public/Model/ProjectLauncherModel.h
+++ b/Engine/Plugins/Developer/ProjectLauncher/Source/ProjectLauncher/Public/Model/ProjectLauncherModel.h
@@ -60,6 +60,9 @@ namespace ProjectLauncher
 	/* Try to parse a string into a content scheme enum value */
 	bool LexTryParseString( ProjectLauncher::EContentScheme& OutContentScheme, const TCHAR* String );
 
+	/* Message box helper */
+	UE_API bool GetUserConfirmation( const FText& Message, const FText& Title = ProjectLauncherText, bool bDefaultValue = true );
+
 	/* Defines the type of a profile (basic, custom, or advanced) */
 	enum class EProfileType : uint8
 	{
