rename BuildCookRunLaunchExtension to BuildCookRunCommandExtension to avoid confusion
- these classes deal with extensions to a BuildCookRun UAT command
- clearing the way for upcoming changes

#rnx

#Synced-CL 49921272

diff --git a/Engine/Plugins/Developer/ProjectLauncher/Source/CommonLaunchExtensions/Private/AdvancedCook/AdvancedCookLaunchExtension.cpp b/Engine/Plugins/Developer/ProjectLauncher/Source/CommonLaunchExtensions/Private/AdvancedCook/AdvancedCookLaunchExtension.cpp
--- a/Engine/Plugins/Developer/ProjectLauncher/Source/CommonLaunchExtensions/Private/AdvancedCook/AdvancedCookLaunchExtension.cpp
+++ b/Engine/Plugins/Developer/ProjectLauncher/Source/CommonLaunchExtensions/Private/AdvancedCook/AdvancedCookLaunchExtension.cpp
@@ -15,7 +15,7 @@
 
 
 FAdvancedCookLaunchExtensionInstance::FAdvancedCookLaunchExtensionInstance( FArgs& InArgs ) 
-	: ProjectLauncher::FBuildCookRunLaunchExtensionInstance(InArgs)
+	: ProjectLauncher::FBuildCookRunCommandExtensionInstance(InArgs)
 	, Owner(StaticCastSharedRef<FAdvancedCookLaunchExtension>(InArgs.Extension))
 {
 }
diff --git a/Engine/Plugins/Developer/ProjectLauncher/Source/CommonLaunchExtensions/Private/AdvancedCook/AdvancedCookLaunchExtension.h b/Engine/Plugins/Developer/ProjectLauncher/Source/CommonLaunchExtensions/Private/AdvancedCook/AdvancedCookLaunchExtension.h
--- a/Engine/Plugins/Developer/ProjectLauncher/Source/CommonLaunchExtensions/Private/AdvancedCook/AdvancedCookLaunchExtension.h
+++ b/Engine/Plugins/Developer/ProjectLauncher/Source/CommonLaunchExtensions/Private/AdvancedCook/AdvancedCookLaunchExtension.h
@@ -2,14 +2,14 @@
 
 #pragma once
 
-#include "Extension/BuildCookRunLaunchExtension.h"
+#include "Extension/BuildCookRunCommandExtension.h"
 
 
 class SCustomLaunchMapListView;
 class SCustomLaunchCultureListView;
 
 
-class FAdvancedCookLaunchExtensionInstance : public ProjectLauncher::FBuildCookRunLaunchExtensionInstance
+class FAdvancedCookLaunchExtensionInstance : public ProjectLauncher::FBuildCookRunCommandExtensionInstance
 {
 public:
 	FAdvancedCookLaunchExtensionInstance( FArgs& InArgs );
@@ -76,7 +76,7 @@ protected:
 };
 
 
-class FAdvancedCookLaunchExtension : public ProjectLauncher::FBuildCookRunLaunchExtension
+class FAdvancedCookLaunchExtension : public ProjectLauncher::FBuildCookRunCommandExtension
 {
 public:
 	virtual TSharedPtr<ProjectLauncher::FLaunchExtensionInstance> CreateInstanceForProfile( ProjectLauncher::FLaunchExtensionInstance::FArgs& InArgs ) override;
diff --git a/Engine/Plugins/Developer/ProjectLauncher/Source/CommonLaunchExtensions/Private/DeprecatedProperties/DeprecatedPropertiesLaunchExtension.cpp b/Engine/Plugins/Developer/ProjectLauncher/Source/CommonLaunchExtensions/Private/DeprecatedProperties/DeprecatedPropertiesLaunchExtension.cpp
--- a/Engine/Plugins/Developer/ProjectLauncher/Source/CommonLaunchExtensions/Private/DeprecatedProperties/DeprecatedPropertiesLaunchExtension.cpp
+++ b/Engine/Plugins/Developer/ProjectLauncher/Source/CommonLaunchExtensions/Private/DeprecatedProperties/DeprecatedPropertiesLaunchExtension.cpp
@@ -14,7 +14,7 @@
 
 
 FDeprecatedPropertiesLaunchExtensionInstance::FDeprecatedPropertiesLaunchExtensionInstance( FArgs& InArgs ) 
-	: ProjectLauncher::FBuildCookRunLaunchExtensionInstance(InArgs)
+	: ProjectLauncher::FBuildCookRunCommandExtensionInstance(InArgs)
 	, Owner(StaticCastSharedRef<FDeprecatedPropertiesLaunchExtension>(InArgs.Extension))
 {
 }
diff --git a/Engine/Plugins/Developer/ProjectLauncher/Source/CommonLaunchExtensions/Private/DeprecatedProperties/DeprecatedPropertiesLaunchExtension.h b/Engine/Plugins/Developer/ProjectLauncher/Source/CommonLaunchExtensions/Private/DeprecatedProperties/DeprecatedPropertiesLaunchExtension.h
--- a/Engine/Plugins/Developer/ProjectLauncher/Source/CommonLaunchExtensions/Private/DeprecatedProperties/DeprecatedPropertiesLaunchExtension.h
+++ b/Engine/Plugins/Developer/ProjectLauncher/Source/CommonLaunchExtensions/Private/DeprecatedProperties/DeprecatedPropertiesLaunchExtension.h
@@ -2,10 +2,10 @@
 
 #pragma once
 
-#include "Extension/BuildCookRunLaunchExtension.h"
+#include "Extension/BuildCookRunCommandExtension.h"
 
 
-class FDeprecatedPropertiesLaunchExtensionInstance : public ProjectLauncher::FBuildCookRunLaunchExtensionInstance
+class FDeprecatedPropertiesLaunchExtensionInstance : public ProjectLauncher::FBuildCookRunCommandExtensionInstance
 {
 public:
 	FDeprecatedPropertiesLaunchExtensionInstance( FArgs& InArgs );
@@ -28,7 +28,7 @@ protected:
 };
 
 
-class FDeprecatedPropertiesLaunchExtension : public ProjectLauncher::FBuildCookRunLaunchExtension
+class FDeprecatedPropertiesLaunchExtension : public ProjectLauncher::FBuildCookRunCommandExtension
 {
 public:
 	virtual TSharedPtr<ProjectLauncher::FLaunchExtensionInstance> CreateInstanceForProfile( ProjectLauncher::FLaunchExtensionInstance::FArgs& InArgs ) override;
diff --git a/Engine/Plugins/Developer/ProjectLauncher/Source/ProjectLauncher/Private/Extension/BuildCookRunCommandExtension.cpp b/Engine/Plugins/Developer/ProjectLauncher/Source/ProjectLauncher/Private/Extension/BuildCookRunCommandExtension.cpp
new file mode 100644
--- /dev/null
+++ b/Engine/Plugins/Developer/ProjectLauncher/Source/ProjectLauncher/Private/Extension/BuildCookRunCommandExtension.cpp
@@ -0,0 +1,257 @@
+// Copyright Epic Games, Inc. All Rights Reserved.
+
+#include "Extension/BuildCookRunCommandExtension.h"
+#include "Model/ProjectLauncherModel.h"
+#include "Styling/ProjectLauncherStyle.h"
+
+#define LOCTEXT_NAMESPACE "FBuildCookRunCommandExtensionInstance"
+
+
+namespace ProjectLauncher
+{
+	FBuildCookRunExtension::FBuildCookRunExtension( const FArgs& InArgs )
+		: Profile(InArgs.Profile)
+		, BuildCookRun(InArgs.BuildCookRun)
+		, DefaultBuildCookRun(InArgs.DefaultBuildCookRun)
+		, Model(InArgs.Model)
+		, Extension(InArgs.Extension)
+		, ExtensionInstance(InArgs.ExtensionInstance)
+		, OwnerTreeData(InArgs.OwnerTreeData)
+	{
+		Profile->OnCustomUATCommandValidation().AddRaw( this, &FBuildCookRunExtension::OnUATCommandValidation );
+		Profile->OnProjectChanged().AddRaw( this, &FBuildCookRunExtension::OnProjectChanged);
+	}
+
+	FBuildCookRunExtension::~FBuildCookRunExtension()
+	{
+		Profile->OnCustomUATCommandValidation().RemoveAll(this);
+		Profile->OnProjectChanged().RemoveAll(this);
+	}
+
+	void FBuildCookRunExtension::OnUATCommandValidation( ILauncherProfileUATCommandRef UATCommand )
+	{
+		if (UATCommand == BuildCookRun && UATCommand->IsEnabled())
+		{
+			OnValidate();
+		}
+	}
+
+	FLaunchProfileTreeNode& FBuildCookRunExtension::AddDefaultHeading(FLaunchProfileTreeNode& ProfileTreeData) const
+	{
+		return AddDefaultHeading(ProfileTreeData, Extension->GetDisplayName() );
+	}
+
+	FLaunchProfileTreeNode& FBuildCookRunExtension::AddDefaultHeading(FLaunchProfileTreeNode& ProfileTreeData, FText OverrideDisplayName) const
+	{
+		if (Model->CanUseSimplifiedLayout(Profile))
+		{
+			return ProfileTreeData.GetTreeData()->AddHeading(Extension->GetInternalName(), OverrideDisplayName, GetHeadingSortOrder() );
+		}
+
+		return ProfileTreeData.AddSubHeading(Extension->GetInternalName(), OverrideDisplayName, GetHeadingSortOrder() );
+	}
+
+
+
+
+	FBuildCookRunCommandExtensionInstance::FBuildCookRunCommandExtensionInstance( FArgs& InArgs )
+		: FLaunchExtensionInstance(InArgs)
+	{
+		GetProfile()->OnUATCommandAdded().AddRaw( this, &FBuildCookRunCommandExtensionInstance::OnUATCommandAdded);
+		GetProfile()->OnUATCommandRemoved().AddRaw( this, &FBuildCookRunCommandExtensionInstance::OnUATCommandRemoved);
+
+	}
+
+	FBuildCookRunCommandExtensionInstance::~FBuildCookRunCommandExtensionInstance()
+	{
+		BuildCookRunExtensions.Reset();
+		GetProfile()->OnUATCommandAdded().RemoveAll(this );
+		GetProfile()->OnUATCommandRemoved().RemoveAll(this );
+	}
+
+	void FBuildCookRunCommandExtensionInstance::InternalInitialize()
+	{
+		FLaunchExtensionInstance::InternalInitialize();
+
+		// create all build cook run extensions if we support them		
+		for (const ILauncherProfileBuildCookRunRef& BuildCookRun : GetProfile()->GetBuildCookRunCommands())
+		{
+			bool bEnabledByDefault = IsBuildCookRunExtensionEnabledByDefault(BuildCookRun);
+			if (GetConfigBool(EConfig::PerProfile, *GetConfigKeyName(TEXT("Enabled"), BuildCookRun), bEnabledByDefault))
+			{
+				CreateBuildCookRunExtensionInternal(BuildCookRun);
+			}
+		}
+	}
+
+	void FBuildCookRunCommandExtensionInstance::CustomizeTree( FLaunchProfileTreeData& ProfileTreeData )
+	{
+		FLaunchExtensionInstance::CustomizeTree(ProfileTreeData);
+
+		for (const TTuple<ILauncherProfileUATCommandPtr,TSharedRef<FBuildCookRunExtension>>& Pair : BuildCookRunExtensions)
+		{
+			const bool bIsPrimaryHeading = true;
+			FLaunchProfileTreeNode& ProfileTreeNode = ProfileTreeData.FindOrAddHeading(Pair.Key.ToSharedRef(), bIsPrimaryHeading);
+
+			Pair.Value->CustomizeTree(ProfileTreeNode);
+		}
+	}
+
+	FString FBuildCookRunCommandExtensionInstance::GetConfigKeyName( const TCHAR* Name, const ILauncherProfileBuildCookRunRef& BuildCookRun) const
+	{
+		return FString::Printf(TEXT("%s.%s"), BuildCookRun->GetInternalName(), Name);
+	}
+
+	bool FBuildCookRunCommandExtensionInstance::HasActiveBuildCookRunExtension() const
+	{
+		return (BuildCookRunExtensions.Num() > 0);
+	}
+
+	TSharedRef<FBuildCookRunExtension> FBuildCookRunCommandExtensionInstance::CreateBuildCookRunExtension( const FBuildCookRunExtension::FArgs& InArgs )
+	{
+		checkNoEntry();
+		TSharedPtr<FBuildCookRunExtension> Invalid;
+		return Invalid.ToSharedRef();
+
+	}
+	
+	TSharedRef<FBuildCookRunExtension> FBuildCookRunCommandExtensionInstance::CreateBuildCookRunExtensionInternal( const ILauncherProfileBuildCookRunRef& InBuildCookRun )
+	{
+		TSharedRef<FBuildCookRunExtension>* ExistingInstancePtr = BuildCookRunExtensions.Find(InBuildCookRun.ToSharedPtr());
+		if (ExistingInstancePtr != nullptr)
+		{
+			return (*ExistingInstancePtr);
+		}
+
+		FBuildCookRunExtension::FArgs Args
+		{
+			.Profile = GetProfile(),
+			.BuildCookRun = InBuildCookRun,
+			.DefaultBuildCookRun = GetModel()->GetCustomDefaultBuildCookRun(),
+			.Model = GetModel(),
+			.Extension = GetExtension(),
+			.ExtensionInstance = AsShared(),
+			.OwnerTreeData = GetOwnerTreeData(),
+		};
+
+		TSharedRef<FBuildCookRunExtension> NewInstance = CreateBuildCookRunExtension(Args);
+
+		BuildCookRunExtensions.Add(InBuildCookRun.ToSharedPtr(), NewInstance);
+
+		return NewInstance;
+	}
+
+	void FBuildCookRunCommandExtensionInstance::CustomizeBuildCookRunExtensionMenu( FMenuBuilder& InMenuBuilder, const ILauncherProfileBuildCookRunRef& InBuildCookRun )
+	{
+		auto IsEnabled = [this, InBuildCookRun]()
+		{
+			return BuildCookRunExtensions.Contains(InBuildCookRun.ToSharedPtr());
+		};
+
+		auto ToggleEnabled = [this, InBuildCookRun, IsEnabled]()
+		{
+			bool bEnabled = !IsEnabled();
+			SetConfigBool(EConfig::PerProfile, *GetConfigKeyName(TEXT("Enabled"), InBuildCookRun), bEnabled);
+			RequestFullTreeRebuild();
+		};
+
+		auto CanToggle = [this, InBuildCookRun, IsEnabled]
+		{
+			bool bWantToEnable = !IsEnabled();
+			return CanToggleBuildCookRunExtension(InBuildCookRun, bWantToEnable);
+		};
+
+		auto GetToolTipText = [InBuildCookRun, CanToggle, IsEnabled]
+		{
+			bool bIsEnabled = IsEnabled();
+			bool bCanToggle = CanToggle();
+
+			if (bIsEnabled && !bCanToggle)
+			{
+				return LOCTEXT("CannotRemoveExtensionTooltip", "The extension cannot be removed at this time");
+			}
+			else if (!bIsEnabled && !bCanToggle)
+			{
+				return LOCTEXT("CannotAddExtensionTooltip", "The extension cannot be added at this time");
+			}
+
+			return FText::GetEmpty();
+		};
+
+
+		InMenuBuilder.AddMenuEntry(
+			GetExtension()->GetDisplayName(),
+			GetToolTipText(),
+			GetExtension()->GetIcon(),
+			FUIAction(
+				FExecuteAction::CreateLambda(ToggleEnabled),
+				FCanExecuteAction::CreateLambda(CanToggle),
+				FIsActionChecked::CreateLambda(IsEnabled)
+			),
+			NAME_None,
+			EUserInterfaceActionType::Check);
+	}
+
+	bool FBuildCookRunCommandExtensionInstance::IsBuildCookRunExtensionEnabledByDefault( const ILauncherProfileBuildCookRunRef& InBuildCookRun ) const
+	{
+		return false;
+	}
+
+	bool FBuildCookRunCommandExtensionInstance::CanToggleBuildCookRunExtension( const ILauncherProfileBuildCookRunRef& InBuildCookRun, bool bWantToEnable ) const
+	{
+		return true;
+	}
+
+
+
+	void FBuildCookRunCommandExtensionInstance::OnUATCommandAdded( const ILauncherProfileUATCommandRef& UATCommand )
+	{
+		ILauncherProfileBuildCookRunPtr BuildCookRun = UATCommand->AsBuildCookRun();
+		if (BuildCookRun.IsValid())
+		{
+			if (!BuildCookRunExtensions.Contains(UATCommand.ToSharedPtr()) && IsBuildCookRunExtensionEnabledByDefault(BuildCookRun.ToSharedRef()))
+			{
+				CreateBuildCookRunExtensionInternal(BuildCookRun.ToSharedRef());
+				RequestFullTreeRebuild();
+			}
+		}
+	}
+
+	void FBuildCookRunCommandExtensionInstance::OnUATCommandRemoved(const ILauncherProfileUATCommandRef& UATCommand)
+	{
+		ILauncherProfileBuildCookRunPtr BuildCookRun = UATCommand->AsBuildCookRun();
+		if (BuildCookRun.IsValid())
+		{
+			if (BuildCookRunExtensions.Contains(UATCommand.ToSharedPtr()))
+			{
+				BuildCookRunExtensions.Remove(UATCommand.ToSharedPtr());
+				RequestFullTreeRebuild();
+			}
+		}
+	}
+
+	bool FBuildCookRunCommandExtensionInstance::CustomizeUATCommandLine( const ILauncherProfileUATCommandRef& InUATCommand, FString& InOutCommandLine )
+	{
+		bool bResult = FLaunchExtensionInstance::CustomizeUATCommandLine(InUATCommand, InOutCommandLine);
+
+		TSharedRef<FBuildCookRunExtension>* ExistingInstancePtr = BuildCookRunExtensions.Find(InUATCommand.ToSharedPtr());
+		if (ExistingInstancePtr != nullptr)
+		{
+			(*ExistingInstancePtr)->CustomizeUATCommandLine(InOutCommandLine);
+			return true;
+		}
+
+		return bResult;
+	}
+
+	void FBuildCookRunCommandExtensionInstance::OnPropertyChanged() 
+	{
+		for (const TTuple<ILauncherProfileUATCommandPtr,TSharedRef<FBuildCookRunExtension>>& Pair : BuildCookRunExtensions)
+		{
+			Pair.Value->OnPropertyChanged();
+		}
+	}
+
+}
+
+#undef LOCTEXT_NAMESPACE
diff --git a/Engine/Plugins/Developer/ProjectLauncher/Source/ProjectLauncher/Private/Extension/BuildCookRunLaunchExtension.cpp b/Engine/Plugins/Developer/ProjectLauncher/Source/ProjectLauncher/Private/Extension/BuildCookRunLaunchExtension.cpp
deleted file mode 100644
diff --git a/Engine/Plugins/Developer/ProjectLauncher/Source/ProjectLauncher/Private/ProfileTree/CustomProfileTreeBuilder.cpp b/Engine/Plugins/Developer/ProjectLauncher/Source/ProjectLauncher/Private/ProfileTree/CustomProfileTreeBuilder.cpp
--- a/Engine/Plugins/Developer/ProjectLauncher/Source/ProjectLauncher/Private/ProfileTree/CustomProfileTreeBuilder.cpp
+++ b/Engine/Plugins/Developer/ProjectLauncher/Source/ProjectLauncher/Private/ProfileTree/CustomProfileTreeBuilder.cpp
@@ -3,7 +3,7 @@
 #include "ProfileTree/CustomProfileTreeBuilder.h"
 #include "Model/ProjectLauncherModel.h"
 #include "Extension/LaunchExtension.h"
-#include "Extension/BuildCookRunLaunchExtension.h"
+#include "Extension/BuildCookRunCommandExtension.h"
 
 
 #define LOCTEXT_NAMESPACE "CustomProfileTreeBuilder"
diff --git a/Engine/Plugins/Developer/ProjectLauncher/Source/ProjectLauncher/Private/Widgets/Profiles/SCustomLaunchCustomProfileEditor.cpp b/Engine/Plugins/Developer/ProjectLauncher/Source/ProjectLauncher/Private/Widgets/Profiles/SCustomLaunchCustomProfileEditor.cpp
--- a/Engine/Plugins/Developer/ProjectLauncher/Source/ProjectLauncher/Private/Widgets/Profiles/SCustomLaunchCustomProfileEditor.cpp
+++ b/Engine/Plugins/Developer/ProjectLauncher/Source/ProjectLauncher/Private/Widgets/Profiles/SCustomLaunchCustomProfileEditor.cpp
@@ -6,7 +6,7 @@
 #include "Styling/StyleColors.h"
 #include "Styling/ProjectLauncherStyle.h"
 #include "Extension/LaunchExtension.h"
-#include "Extension/BuildCookRunLaunchExtension.h"
+#include "Extension/BuildCookRunCommandExtension.h"
 #include "Framework/Commands/UIAction.h"
 #include "Framework/MultiBox/MultiBoxBuilder.h"
 #include "SlateOptMacros.h"
diff --git a/Engine/Plugins/Developer/ProjectLauncher/Source/ProjectLauncher/Public/Extension/BuildCookRunCommandExtension.h b/Engine/Plugins/Developer/ProjectLauncher/Source/ProjectLauncher/Public/Extension/BuildCookRunCommandExtension.h
new file mode 100644
--- /dev/null
+++ b/Engine/Plugins/Developer/ProjectLauncher/Source/ProjectLauncher/Public/Extension/BuildCookRunCommandExtension.h
@@ -0,0 +1,215 @@
+// Copyright Epic Games, Inc. All Rights Reserved.
+
+#pragma once
+
+#include "Extension/LaunchExtension.h"
+
+#define UE_API PROJECTLAUNCHER_API
+
+namespace ProjectLauncher
+{
+	/**
+	 * Base for a launch extension for an ILauncherProfileBuildCookRun
+	 */
+	class FBuildCookRunExtension : public TSharedFromThis<FBuildCookRunExtension>
+	{
+	public:
+		struct FArgs
+		{
+			ILauncherProfileRef Profile;
+			ILauncherProfileBuildCookRunRef BuildCookRun;
+			ILauncherProfileBuildCookRunRef DefaultBuildCookRun;
+			TSharedRef<FModel> Model;
+			TSharedRef<FLaunchExtension> Extension;
+			TSharedRef<FLaunchExtensionInstance> ExtensionInstance;
+			FLaunchProfileTreeData& OwnerTreeData;
+		};
+
+		UE_API FBuildCookRunExtension( const FArgs& InArgs );
+		UE_API virtual ~FBuildCookRunExtension();
+
+		/**
+		 * Hook to allow the extension to add extra fields to the property editing tree, if the tree builder allows it
+		 * Property tree items should be hidden until the user has selected something to make it relevent, to avoid cluttering the UI.
+		 */
+		virtual void CustomizeTree( FLaunchProfileTreeNode& ProfileTreeNode ) {};
+
+		/**
+		 * Create the default heading for this extension
+		 */
+		UE_API FLaunchProfileTreeNode& AddDefaultHeading(FLaunchProfileTreeNode& ProfileTreeNode) const;
+
+		/**
+		 * Create the default heading for this extension, with a bespoke display name
+		 */
+		UE_API FLaunchProfileTreeNode& AddDefaultHeading(FLaunchProfileTreeNode& ProfileTreeNode, FText OverrideDisplayName) const;
+
+
+
+		/**
+		 * Advanced hook to allow any advanced modification of the command line when a UAT command test is launched
+		 */
+		virtual void CustomizeUATCommandLine( FString& InOutCommandLine ) {};
+
+		/**
+		 * Notification callback when a property has been changed in our profile
+		 */
+		void OnPropertyChanged() {};
+
+		/**
+		 * Notification callback when the project has been changed in our profile
+		 */
+		virtual void OnProjectChanged() {};
+
+		/**
+		 * Notification callback when the profile is validated
+		 */
+		virtual void OnValidate() {};
+
+		/** 
+		 * Get the extension we belong to
+		 */
+		inline TSharedRef<FLaunchExtension> GetExtension() const { return Extension; }
+
+		/** 
+		 * Get the extension instance that instantiated us
+		 */
+		inline TSharedRef<FLaunchExtensionInstance> GetExtensionInstance() const { return ExtensionInstance; }
+
+	protected:
+		/** 
+		 * Get the profile we were instantiated for
+		 */
+		inline ILauncherProfileRef GetProfile() const { return Profile; }
+
+		/** 
+		 * Get the BuildCookRun we were instantiated for
+		 */
+		inline ILauncherProfileBuildCookRunRef GetBuildCookRun() const { return BuildCookRun; }
+
+		/** 
+		 * Get the reference default BuildCookRun
+		 */
+		inline ILauncherProfileBuildCookRunRef GetDefaultBuildCookRun() const { return DefaultBuildCookRun; }
+		
+		/** 
+		 * Get the model, for general purpose helper functions
+		 */
+		inline TSharedRef<FModel> GetModel() const { return Model; }
+
+	private:
+		void OnUATCommandValidation( ILauncherProfileUATCommandRef UATCommand );
+		virtual int32 GetHeadingSortOrder() const { return 0; }
+
+		ILauncherProfileRef Profile;
+		ILauncherProfileBuildCookRunRef BuildCookRun;
+		ILauncherProfileBuildCookRunRef DefaultBuildCookRun;
+		TSharedRef<FModel> Model;
+		TSharedRef<FLaunchExtension> Extension;
+		TSharedRef<FLaunchExtensionInstance> ExtensionInstance;
+		FLaunchProfileTreeData& OwnerTreeData;
+		FPropertyChangedDelegate PropertyChangedDelegate;
+
+
+	};
+
+
+	/**
+	 * Interface for creating & maintaining build cook run extensions
+	 */
+	class IBuildCookRunExtensionFactory
+	{
+	public:
+
+		/**
+		 * A launch extension must override this to create a new build cook run extension
+		 */
+		virtual TSharedRef<FBuildCookRunExtension> CreateBuildCookRunExtension( const FBuildCookRunExtension::FArgs& InArgs ) = 0;
+
+		/**
+		 * Returns whether there are any active build cook run extensions from this launch extension
+		 */
+		virtual bool HasActiveBuildCookRunExtension() const = 0;
+
+		/**
+		 * Customizes the build cook run extensions menu to allow extenions to be enabled & disabled 
+		 */
+		virtual void CustomizeBuildCookRunExtensionMenu( FMenuBuilder& InMenuBuilder, const ILauncherProfileBuildCookRunRef& InBuildCookRun ) = 0;
+
+	protected:
+		/*
+		 * Returns whether build cook run extensions should be enabled by default. Typically false
+		 */
+		virtual bool IsBuildCookRunExtensionEnabledByDefault( const ILauncherProfileBuildCookRunRef& InBuildCookRun ) const = 0;
+
+		/*
+		 * Returns whether our build cook run extensions can be toggled on or off. Typically true
+		 */
+		virtual bool CanToggleBuildCookRunExtension( const ILauncherProfileBuildCookRunRef& InBuildCookRun, bool bWantToEnable ) const = 0;
+	};
+
+
+
+
+	class FBuildCookRunCommandExtensionInstance : public FLaunchExtensionInstance, public IBuildCookRunExtensionFactory
+	{
+	public:
+
+		UE_API FBuildCookRunCommandExtensionInstance( FArgs& InArgs);
+		UE_API virtual ~FBuildCookRunCommandExtensionInstance();
+
+		UE_API virtual void InternalInitialize() override;
+
+
+		UE_API virtual IBuildCookRunExtensionFactory* AsBuildCookRunFactory() override { return this; }
+
+		UE_API bool HasActiveBuildCookRunExtension() const override;
+
+		UE_API void CustomizeBuildCookRunExtensionMenu( FMenuBuilder& InMenuBuilder, const ILauncherProfileBuildCookRunRef& InBuildCookRun )  override;
+
+		UE_API virtual TSharedRef<FBuildCookRunExtension> CreateBuildCookRunExtension( const FBuildCookRunExtension::FArgs& InArgs ) = 0;
+
+		UE_API virtual void CustomizeTree( FLaunchProfileTreeData& ProfileTreeData );
+		UE_API virtual void OnPropertyChanged() override;
+
+		// fixme: ideally enumerate FLaunchExtension to find those that are BCR
+		virtual bool CanBeRemoved() const override final { return true; }
+
+	protected:
+		/**
+		 * Get the key name to use for reading & writing a configuration value for the given BuildCookRun, expected to be passed as Name to the GetConfigXXX/SetConfigXXX functions
+		 * 
+		 * @param Config where the value is stored
+		 * @param Name name of the value
+		 * @param BuildCookRun the build cook run command
+		 * 
+		 * @returns namespaced version of Name
+		 */
+		UE_API FString GetConfigKeyName( const TCHAR* Name, const ILauncherProfileBuildCookRunRef& BuildCookRun) const;
+
+		UE_API virtual bool IsBuildCookRunExtensionEnabledByDefault( const ILauncherProfileBuildCookRunRef& InBuildCookRun ) const override;
+		UE_API virtual bool CanToggleBuildCookRunExtension( const ILauncherProfileBuildCookRunRef& InBuildCookRun, bool bWantToEnable ) const override;
+		UE_API virtual bool CustomizeUATCommandLine( const ILauncherProfileUATCommandRef& InUATCommand, FString& InOutCommandLine ) override;
+
+		UE_API void OnUATCommandAdded(const ILauncherProfileUATCommandRef& UATCommand);
+		UE_API void OnUATCommandRemoved(const ILauncherProfileUATCommandRef& UATCommand);
+		UE_API TSharedRef<FBuildCookRunExtension> CreateBuildCookRunExtensionInternal( const ILauncherProfileBuildCookRunRef& InBuildCookRun );
+
+		TMap<ILauncherProfileUATCommandPtr,TSharedRef<FBuildCookRunExtension>> BuildCookRunExtensions;
+
+		friend class FBuildCookRunExtension;
+	};
+
+
+
+	class FBuildCookRunCommandExtension : public FLaunchExtension
+	{
+	public:
+		virtual bool IsAlwaysCreated( ILauncherProfileRef InProfile, TSharedRef<FModel> InModel ) const override { return true; }
+		virtual class FBuildCookRunCommandExtension* AsBuildCookRunCommandExtension() override { return this; }
+	};
+
+
+};
+
+#undef UE_API
diff --git a/Engine/Plugins/Developer/ProjectLauncher/Source/ProjectLauncher/Public/Extension/BuildCookRunLaunchExtension.h b/Engine/Plugins/Developer/ProjectLauncher/Source/ProjectLauncher/Public/Extension/BuildCookRunLaunchExtension.h
deleted file mode 100644
diff --git a/Engine/Plugins/Developer/ProjectLauncher/Source/ProjectLauncher/Public/Extension/LaunchExtension.h b/Engine/Plugins/Developer/ProjectLauncher/Source/ProjectLauncher/Public/Extension/LaunchExtension.h
--- a/Engine/Plugins/Developer/ProjectLauncher/Source/ProjectLauncher/Public/Extension/LaunchExtension.h
+++ b/Engine/Plugins/Developer/ProjectLauncher/Source/ProjectLauncher/Public/Extension/LaunchExtension.h
@@ -431,7 +431,7 @@ namespace ProjectLauncher
 		/**
 		 * Attempt to cast this launch extension to an BuildCookRun launch extension
 		 */
-		virtual class FBuildCookRunLaunchExtension* AsBuildCookRunLaunchExtension() { return nullptr; }
+		virtual class FBuildCookRunCommandExtension* AsBuildCookRunCommandExtension() { return nullptr; }
 
 		/**
 		 * Instantiate all compatible extensions
