crash fixes when launching a profile with no build cook runs

#rnx

#Synced-CL 49921272

diff --git a/Engine/Plugins/Developer/ProjectLauncher/Source/ProjectLauncher/Private/Widgets/Output/SCustomLaunchLaunchPanel.cpp b/Engine/Plugins/Developer/ProjectLauncher/Source/ProjectLauncher/Private/Widgets/Output/SCustomLaunchLaunchPanel.cpp
--- a/Engine/Plugins/Developer/ProjectLauncher/Source/ProjectLauncher/Private/Widgets/Output/SCustomLaunchLaunchPanel.cpp
+++ b/Engine/Plugins/Developer/ProjectLauncher/Source/ProjectLauncher/Private/Widgets/Output/SCustomLaunchLaunchPanel.cpp
@@ -491,9 +491,9 @@ FText SCustomLaunchLaunchPanel::GetProfileProjectName() const
 FText SCustomLaunchLaunchPanel::GetProfileConfigurationName() const
 {
 	ILauncherProfilePtr LauncherProfile = GetLauncherProfile();
-	if (LauncherProfile != nullptr)
+	if (LauncherProfile != nullptr && LauncherProfile->GetFirstBuildCookRun().IsValid())
 	{
-		return EBuildConfigurations::ToText(LauncherProfile->GetBuildConfiguration());
+		return EBuildConfigurations::ToText(LauncherProfile->GetFirstBuildCookRun()->GetBuildConfiguration());
 	}
 
 	return FText::GetEmpty();
@@ -501,9 +501,9 @@ FText SCustomLaunchLaunchPanel::GetProfileConfigurationName() const
 FText SCustomLaunchLaunchPanel::GetProfileTargetName() const
 {
 	ILauncherProfilePtr LauncherProfile = GetLauncherProfile();
-	if (LauncherProfile != nullptr)
+	if (LauncherProfile != nullptr && LauncherProfile->GetFirstBuildCookRun().IsValid())
 	{
-		TArray<FString> BuildTargets = LauncherProfile->GetBuildTargets();
+		TArray<FString> BuildTargets = LauncherProfile->GetFirstBuildCookRun()->GetBuildTargets();
 		if (BuildTargets.Num() == 0 || BuildTargets.Contains(FString()))
 		{
 			return LOCTEXT("DefaultTargetName", "Target Default");
@@ -531,15 +531,16 @@ FText SCustomLaunchLaunchPanel::GetProfileContentSchemeName() const
 		{
 			return LOCTEXT("MultipleContentSchemeLabel", "Multiple"); // @todo: worth trying something fancy here? collecting all content schemes we're using
 		}
-		else if (!LauncherProfile->GetFirstBuildCookRun().IsValid())
-		{
-			return LOCTEXT("NoBuildCookRunsLabel", "Custom Profile");
-		}
-		else
+		else if (LauncherProfile->GetFirstBuildCookRun().IsValid())
 		{
 			ProjectLauncher::EContentScheme ContentScheme = Model->DetermineProfileContentScheme(LauncherProfile->GetFirstBuildCookRun().ToSharedRef());
 			return ProjectLauncher::GetContentSchemeDisplayName(ContentScheme);
 		}
+		else if (LauncherProfile->GetUATCommands().Num() > 0)
+		{
+			return LOCTEXT("NoBuildCookRunsLabel", "Custom Profile");
+		}
+
 	}
 
 	return FText::GetEmpty();
diff --git a/Engine/Source/Developer/LauncherServices/Private/Launcher/LauncherWorker.cpp b/Engine/Source/Developer/LauncherServices/Private/Launcher/LauncherWorker.cpp
--- a/Engine/Source/Developer/LauncherServices/Private/Launcher/LauncherWorker.cpp
+++ b/Engine/Source/Developer/LauncherServices/Private/Launcher/LauncherWorker.cpp
@@ -1676,7 +1676,7 @@ void FLauncherWorker::CreateAndExecuteTasks( const ILauncherProfileRef& InProfil
 		}
 	};
 
-	if (InProfile->GetCookMode() == ELauncherProfileCookModes::ByTheBookInEditor)
+	if (InProfile->GetFirstBuildCookRun().IsValid() && InProfile->GetFirstBuildCookRun()->GetCookMode() == ELauncherProfileCookModes::ByTheBookInEditor)
 	{
 		// need a command which will wait for the cook to finish
 		class FWaitForCookInEditorToFinish : public FLauncherTask
