move extension & tasks toolbar to the top of the property edit window + make the task menu a 'positive action' menu

#rnx

#Synced-CL 49921272

diff --git a/Engine/Plugins/Developer/ProjectLauncher/Source/ProjectLauncher/Private/Widgets/Profiles/SCustomLaunchProfilesPanel.cpp b/Engine/Plugins/Developer/ProjectLauncher/Source/ProjectLauncher/Private/Widgets/Profiles/SCustomLaunchProfilesPanel.cpp
--- a/Engine/Plugins/Developer/ProjectLauncher/Source/ProjectLauncher/Private/Widgets/Profiles/SCustomLaunchProfilesPanel.cpp
+++ b/Engine/Plugins/Developer/ProjectLauncher/Source/ProjectLauncher/Private/Widgets/Profiles/SCustomLaunchProfilesPanel.cpp
@@ -255,7 +255,7 @@ TSharedRef<SWidget> SCustomLaunchProfilesPanel::CreateProfileSelectorWidget()
 			+SHorizontalBox::Slot()
 			.FillWidth(1)
 			.VAlign(VAlign_Center)
-			.Padding(4,2)
+			.Padding(6,2)
 			[
 				SNew(STextBlock)
 				.Text(LOCTEXT("ProfilesLabel", "Launch Profiles"))
@@ -325,23 +325,6 @@ TSharedRef<SWidget> SCustomLaunchProfilesPanel::CreateProfileEditorWidget()
 	return SNew(SVerticalBox)
 	.Visibility_Lambda( [this] { return Model->GetSelectedProfile().IsValid() && IsProfileEditorVisible() ? EVisibility::Visible :EVisibility::Collapsed; } )
 
-	// advanced profile warning banner
-	+SVerticalBox::Slot()
-	.AutoHeight()
-	[
-		CreateAdvancedProfileWarningWidget()
-	]
-
-	// profile property editor
-	+SVerticalBox::Slot()
-	.FillHeight(1)
-	[
-		SNew(SScrollBorder, PropertyEditor.ToSharedRef())
-		[
-			PropertyEditor.ToSharedRef()
-		]
-	]
-
 	// extensions button bar
 	+SVerticalBox::Slot()
 	.AutoHeight()
@@ -352,11 +335,22 @@ TSharedRef<SWidget> SCustomLaunchProfilesPanel::CreateProfileEditorWidget()
 		[
 			SNew(SHorizontalBox)
 
+			// bannel label
+			+SHorizontalBox::Slot()
+			.FillWidth(1)
+			.VAlign(VAlign_Center)
+			.Padding(6,2)
+			[
+				SNew(STextBlock)
+				.Text(LOCTEXT("TasksLabel", "Launch Tasks"))
+			]
+
+			// extensions menu
 			+SHorizontalBox::Slot()
 			.AutoWidth()
 			.VAlign(VAlign_Center)
-			.HAlign(HAlign_Left)
-			.Padding(4,4)
+			.HAlign(HAlign_Right)
+			.Padding(4,2)
 			[
 				SNew(SComboButton)
 				.ComboButtonStyle(FAppStyle::Get(), "SimpleComboButton")
@@ -368,7 +362,8 @@ TSharedRef<SWidget> SCustomLaunchProfilesPanel::CreateProfileEditorWidget()
 					SNew(SHorizontalBox)
 					+SHorizontalBox::Slot()
 					.AutoWidth()
-					.Padding(2, 0)
+					.VAlign(VAlign_Center)
+					.Padding(1,3)
 					[
 						SNew(SImage)
 						.Image(FProjectLauncherStyle::Get().GetBrush("Icons.Plugins"))
@@ -376,7 +371,8 @@ TSharedRef<SWidget> SCustomLaunchProfilesPanel::CreateProfileEditorWidget()
 					]
 					+SHorizontalBox::Slot()
 					.AutoWidth()
-					.Padding(2, 0)
+					.VAlign(VAlign_Center)
+					.Padding(1,3)
 					[
 						SNew(STextBlock)
 						.Text(LOCTEXT("ExtensionsLabel", "Extensions"))
@@ -385,44 +381,62 @@ TSharedRef<SWidget> SCustomLaunchProfilesPanel::CreateProfileEditorWidget()
 				]
 			]
 
+			// new task button
 			+SHorizontalBox::Slot()
 			.AutoWidth()
 			.VAlign(VAlign_Center)
-			.HAlign(HAlign_Left)
-			.Padding(4,4)
+			.HAlign(HAlign_Right)
+			.Padding(4,2)
 			[
 				SNew(SComboButton)
-				.ComboButtonStyle(FAppStyle::Get(), "SimpleComboButton")
+				.ComboButtonStyle(FAppStyle::Get(), "ComboButton")
+				.ButtonStyle(FAppStyle::Get(), "Button")
 				.ToolTipText(LOCTEXT("CommandsToolTip", "Configure the commands to run."))
 				.OnGetMenuContent(PropertyEditor.ToSharedRef(), &SCustomLaunchCustomProfileEditor::MakeCommandsMenu)
 				.MenuPlacement(MenuPlacement_ComboBoxRight)
+				.HasDownArrow(false)
 				.ButtonContent()
 				[
 					SNew(SHorizontalBox)
+
 					+SHorizontalBox::Slot()
 					.AutoWidth()
-					.Padding(2, 0)
+					.VAlign(VAlign_Center)
+					.Padding(1,3)
 					[
 						SNew(SImage)
-						.Image(FProjectLauncherStyle::Get().GetBrush("Icons.Extensions"))
-						.ColorAndOpacity(FSlateColor::UseForeground())
+						.Image(FAppStyle::Get().GetBrush("Icons.Plus"))
+						.ColorAndOpacity(FStyleColors::AccentGreen)
 					]
+
 					+SHorizontalBox::Slot()
 					.AutoWidth()
-					.Padding(2, 0)
+					.VAlign(VAlign_Center)
+					.Padding(1,3)
 					[
 						SNew(STextBlock)
-						.Text(LOCTEXT("CommandsLabel", "Tasks"))
-						.ColorAndOpacity(FSlateColor::UseForeground())
+						.Text(LOCTEXT("CommandsLabel", "Add Task"))
 					]
 				]
 			]
+		]
+	]
 
-			+SHorizontalBox::Slot()
-			.FillWidth(1)
-			[
-				SNew(SSpacer)
-			]
+
+	// advanced profile warning banner
+	+SVerticalBox::Slot()
+	.AutoHeight()
+	[
+		CreateAdvancedProfileWarningWidget()
+	]
+
+	// profile property editor
+	+SVerticalBox::Slot()
+	.FillHeight(1)
+	[
+		SNew(SScrollBorder, PropertyEditor.ToSharedRef())
+		[
+			PropertyEditor.ToSharedRef()
 		]
 	]
 	;
