Add support for multiple UAT command extension instances
- simplified UAT & Automated Test extension api
- unified menu handling between all extension types

#rnx

#Synced-CL 50040175

diff --git a/Engine/Plugins/Developer/ProjectLauncher/Source/CommonLaunchExtensions/Private/BootTest/BootTestLaunchExtension.cpp b/Engine/Plugins/Developer/ProjectLauncher/Source/CommonLaunchExtensions/Private/BootTest/BootTestLaunchExtension.cpp
--- a/Engine/Plugins/Developer/ProjectLauncher/Source/CommonLaunchExtensions/Private/BootTest/BootTestLaunchExtension.cpp
+++ b/Engine/Plugins/Developer/ProjectLauncher/Source/CommonLaunchExtensions/Private/BootTest/BootTestLaunchExtension.cpp
@@ -5,9 +5,6 @@
 
 #define LOCTEXT_NAMESPACE "FBootTestLaunchExtensionInstance"
 
-const TCHAR* FBootTestLaunchExtensionInstance::BootTestInternalName = TEXT("BootTestExtension.BootTest");
-
-
 void FBootTestLaunchExtensionInstance::CustomizeTree( ProjectLauncher::FLaunchProfileTreeData& ProfileTreeData )
 {
 	// note: Windowed option is here mostly as an example of adding UI items for a particular profile - your tests may not need it
@@ -21,7 +18,7 @@ void FBootTestLaunchExtensionInstance::CustomizeTree( ProjectLauncher::FLaunchPr
 	);
 }
 
-void FBootTestLaunchExtensionInstance::CustomizeAutomatedTestCommandLine( FString& InOutCommandLine )
+void FBootTestLaunchExtensionInstance::CustomizeUATCommandLine( FString& InOutCommandLine )
 {
 	bool bWindowed = GetConfigBool(EConfig::PerProfile, TEXT("Windowed"));
 	if (!bWindowed)
@@ -30,23 +27,11 @@ void FBootTestLaunchExtensionInstance::CustomizeAutomatedTestCommandLine( FStrin
 	}
 }
 
-const FString FBootTestLaunchExtensionInstance::GetTestInternalName() const
-{
-	return BootTestInternalName;
-}
-
 void FBootTestLaunchExtensionInstance::OnTestAdded( ILauncherProfileAutomatedTestRef AutomatedTest )
 {
 	AutomatedTest->SetTests(TEXT("UE.BootTest"));
 }
 
-int32 FBootTestLaunchExtensionInstance::GetTestOrder() const
-{
-	// Boot Test is likely the most lightweight test there is, so we want to run it first - give it a low order
-	return 1; 
-}
-
-
 TSharedPtr<ProjectLauncher::FLaunchExtensionInstance> FBootTestLaunchExtension::CreateInstanceForProfile( ProjectLauncher::FLaunchExtensionInstance::FArgs& InArgs )
 {
 	return MakeShared<FBootTestLaunchExtensionInstance>(InArgs);
diff --git a/Engine/Plugins/Developer/ProjectLauncher/Source/CommonLaunchExtensions/Private/BootTest/BootTestLaunchExtension.h b/Engine/Plugins/Developer/ProjectLauncher/Source/CommonLaunchExtensions/Private/BootTest/BootTestLaunchExtension.h
--- a/Engine/Plugins/Developer/ProjectLauncher/Source/CommonLaunchExtensions/Private/BootTest/BootTestLaunchExtension.h
+++ b/Engine/Plugins/Developer/ProjectLauncher/Source/CommonLaunchExtensions/Private/BootTest/BootTestLaunchExtension.h
@@ -13,11 +13,9 @@ public:
 	virtual ~FBootTestLaunchExtensionInstance() = default;
 
 	virtual void OnTestAdded( ILauncherProfileAutomatedTestRef AutomatedTest ) override;
-	virtual const FString GetTestInternalName() const override;
-	virtual int32 GetTestOrder() const override;
 
 	virtual void CustomizeTree( ProjectLauncher::FLaunchProfileTreeData& ProfileTreeData ) override;
-	virtual void CustomizeAutomatedTestCommandLine( FString& InOutCommandLine ) override;
+	virtual void CustomizeUATCommandLine( FString& InOutCommandLine ) override;
 
 private:
 	static const TCHAR* BootTestInternalName;
diff --git a/Engine/Plugins/Developer/ProjectLauncher/Source/ProjectLauncher/Private/Extension/AutomatedTestLaunchExtension.cpp b/Engine/Plugins/Developer/ProjectLauncher/Source/ProjectLauncher/Private/Extension/AutomatedTestLaunchExtension.cpp
--- a/Engine/Plugins/Developer/ProjectLauncher/Source/ProjectLauncher/Private/Extension/AutomatedTestLaunchExtension.cpp
+++ b/Engine/Plugins/Developer/ProjectLauncher/Source/ProjectLauncher/Private/Extension/AutomatedTestLaunchExtension.cpp
@@ -1,6 +1,7 @@
 // Copyright Epic Games, Inc. All Rights Reserved.
 
 #include "Extension/AutomatedTestLaunchExtension.h"
+#include "Model/ProjectLauncherModel.h"
 #include "Styling/ProjectLauncherStyle.h"
 
 #define LOCTEXT_NAMESPACE "FAutomatedTestLaunchExtensionInstance"
@@ -9,62 +10,24 @@
 namespace ProjectLauncher
 {
 	FAutomatedTestLaunchExtensionInstance::FAutomatedTestLaunchExtensionInstance( FArgs& InArgs )
-		: FLaunchExtensionInstance(InArgs)
+		: FUATCommandLaunchExtensionInstanceBase(InArgs)
 	{
+		check(UATCommand->AsAutomatedTest().IsValid());
 	}
 
-	void FAutomatedTestLaunchExtensionInstance::InternalInitialize()
-	{
-	}
-
-	void FAutomatedTestLaunchExtensionInstance::OnCustomValidation()
-	{
-		ILauncherProfileAutomatedTestPtr AutomatedTest = GetTest();
-		if (AutomatedTest.IsValid() && AutomatedTest->IsEnabled())
-		{
-			OnValidateProfile();
-		}
-	}
-
-	bool FAutomatedTestLaunchExtensionInstance::ManagesUATCommand( const ILauncherProfileUATCommandRef& UATCommand ) const
-	{
-		return UATCommand == GetTest();
-	}
-
-
-
-	FLaunchProfileTreeNode& FAutomatedTestLaunchExtensionInstance::AddDefaultHeadingInternal(ProjectLauncher::FLaunchProfileTreeData& ProfileTreeData, FText DisplayName, int32 SortOrder) const
+	void FAutomatedTestLaunchExtensionInstance::OnAdded()
 	{
 		ILauncherProfileAutomatedTestPtr AutomatedTest = GetTest();
-		if (AutomatedTest.IsValid())
-		{
-			return ProfileTreeData.FindOrAddHeading(AutomatedTest.ToSharedRef(), true );
-		}
 
-		return FLaunchExtensionInstance::AddDefaultHeadingInternal(ProfileTreeData, DisplayName, SortOrder);
-	}
+		// set the defaults, placing automated tests at the end by default
+		FModel::FUATCommandsDetail UATCommandsDetail; 
+		GetModel()->GetUATCommandsDetail(GetProfile(), UATCommandsDetail);
+		AutomatedTest->SetOrder( UATCommandsDetail.MaxOrder + 1 );
+		AutomatedTest->SetDescription( GetExtension()->GetDisplayName().ToString() );
 
-
-	void FAutomatedTestLaunchExtensionInstance::OnAdded()
-	{
-		bool bIsNew = false;
-		ILauncherProfileAutomatedTestPtr AutomatedTest = GetTest();
-		if (!AutomatedTest.IsValid())
-		{
-			AutomatedTest = GetProfile()->FindOrAddAutomatedTest(*GetTestInternalName());
-			bIsNew = true;
-		}
 		OnTestAdded(AutomatedTest.ToSharedRef());
 
-		if (bIsNew)
-		{
-			int32 Order = GetTestOrder();
-			AutomatedTest->SetOrder(Order);
-			AutomatedTest->SetDescription( GetExtension()->GetDisplayName().ToString() );
-		}
-
 		BroadcastPropertyChanged();
-
 	}
 
 	void FAutomatedTestLaunchExtensionInstance::OnRemoved()
@@ -73,33 +36,17 @@ namespace ProjectLauncher
 		if (ExistingAutomatedTest.IsValid())
 		{
 			OnTestRemoved(ExistingAutomatedTest.ToSharedRef());
-			GetProfile()->RemoveAutomatedTest(*GetTestInternalName());
+			GetProfile()->RemoveAutomatedTest(ExistingAutomatedTest->GetInternalName());
 		}
 		BroadcastPropertyChanged();
-		RequestFullTreeRebuild(); // @fixme: still needed?
-	}
-
-	bool FAutomatedTestLaunchExtensionInstance::IsTestActive() const
-	{
-		return GetTest() != nullptr;
 	}
 
 	ILauncherProfileAutomatedTestPtr FAutomatedTestLaunchExtensionInstance::GetTest() const
 	{
-		return GetProfile()->GetAutomatedTest(*GetTestInternalName());
+		return UATCommand->AsAutomatedTest();
 	}
 
 
-	bool FAutomatedTestLaunchExtensionInstance::CustomizeAutomatedTestCommandLine( const ILauncherProfileAutomatedTestRef& InAutomatedTest, FString& InOutCommandLine )
-	{
-		if (InAutomatedTest == GetTest())
-		{
-			CustomizeAutomatedTestCommandLine(InOutCommandLine);
-			return true;
-		}
-
-		return false;
-	}
 
 
 
@@ -113,6 +60,12 @@ namespace ProjectLauncher
 		return FSlateIcon( FProjectLauncherStyle::GetStyleSetName(), "Icons.Task.TestAutomation");
 	}
 
+	ILauncherProfileUATCommandRef FAutomatedTestLaunchExtension::CreateUATCommand(ILauncherProfileRef InProfile, TSharedRef<FModel> InModel) const
+	{
+		return InProfile->FindOrAddAutomatedTest( nullptr, GetInternalName() );
+	}
+
+
 }
 
 #undef LOCTEXT_NAMESPACE
diff --git a/Engine/Plugins/Developer/ProjectLauncher/Source/ProjectLauncher/Private/Extension/BuildCookRunCommandExtension.cpp b/Engine/Plugins/Developer/ProjectLauncher/Source/ProjectLauncher/Private/Extension/BuildCookRunCommandExtension.cpp
--- a/Engine/Plugins/Developer/ProjectLauncher/Source/ProjectLauncher/Private/Extension/BuildCookRunCommandExtension.cpp
+++ b/Engine/Plugins/Developer/ProjectLauncher/Source/ProjectLauncher/Private/Extension/BuildCookRunCommandExtension.cpp
@@ -45,10 +45,10 @@ namespace ProjectLauncher
 	{
 		if (Model->CanUseSimplifiedLayout(Profile))
 		{
-			return ProfileTreeData.GetTreeData()->AddHeading(Extension->GetInternalName(), OverrideDisplayName, GetHeadingSortOrder() );
+			return ProfileTreeData.GetTreeData()->AddHeading(Extension->GetInternalName(), OverrideDisplayName, BuildCookRun->GetOrder() );
 		}
 
-		return ProfileTreeData.AddSubHeading(Extension->GetInternalName(), OverrideDisplayName, GetHeadingSortOrder() );
+		return ProfileTreeData.AddSubHeading(Extension->GetInternalName(), OverrideDisplayName, BuildCookRun->GetOrder() );
 	}
 
 
@@ -251,7 +251,6 @@ namespace ProjectLauncher
 			Pair.Value->OnPropertyChanged();
 		}
 	}
-
 }
 
 #undef LOCTEXT_NAMESPACE
diff --git a/Engine/Plugins/Developer/ProjectLauncher/Source/ProjectLauncher/Private/Extension/BuildCookRunLaunchExtension.cpp b/Engine/Plugins/Developer/ProjectLauncher/Source/ProjectLauncher/Private/Extension/BuildCookRunLaunchExtension.cpp
new file mode 100644
--- /dev/null
+++ b/Engine/Plugins/Developer/ProjectLauncher/Source/ProjectLauncher/Private/Extension/BuildCookRunLaunchExtension.cpp
@@ -0,0 +1,62 @@
+// Copyright Epic Games, Inc. All Rights Reserved.
+
+#include "Extension/BuildCookRunLaunchExtension.h"
+#include "Model/ProjectLauncherModel.h"
+#include "Styling/ProjectLauncherStyle.h"
+
+#define LOCTEXT_NAMESPACE "FBuildCookRunLaunchExtensionInstance"
+
+
+namespace ProjectLauncher
+{
+	FBuildCookRunLaunchExtensionInstance::FBuildCookRunLaunchExtensionInstance( FArgs& InArgs )
+		: FUATCommandLaunchExtensionInstanceBase(InArgs)
+	{
+		check(UATCommand->AsBuildCookRun().IsValid());
+	}
+
+
+	
+	void FBuildCookRunLaunchExtension::GetExtensionsMenuEntry( FExtensionsMenuEntry& MenuEntry ) const
+	{
+		MenuEntry = FExtensionsMenuEntry::Default;
+	}
+
+	FSlateIcon FBuildCookRunLaunchExtension::GetIcon() const
+	{
+		return FSlateIcon(FProjectLauncherStyle::GetStyleSetName(), "Icons.Task.Launch");
+	}
+
+	ILauncherProfileUATCommandRef FBuildCookRunLaunchExtension::CreateUATCommand(ILauncherProfileRef InProfile, TSharedRef<FModel> InModel) const
+	{
+		ILauncherProfileBuildCookRunRef BuildCookRun = InProfile->FindOrAddBuildCookRunCommand( nullptr, GetInternalName() );
+
+		InModel->SetDefaults(InProfile, BuildCookRun);
+
+		return BuildCookRun;
+
+	}
+
+	const TCHAR* FBuildCookRunLaunchExtension::GetInternalName() const
+	{
+		return TEXT("BuildCookRun");
+	}
+
+	FText FBuildCookRunLaunchExtension::GetDisplayName() const
+	{
+		return LOCTEXT("BuildCookRunExtensionName", "Build, Cook & Run");
+	}
+
+	TSharedPtr<ProjectLauncher::FLaunchExtensionInstance> FBuildCookRunLaunchExtension::CreateInstanceForProfile( FLaunchExtensionInstance::FArgs& InArgs )
+	{
+		return MakeShared<FBuildCookRunLaunchExtensionInstance>(InArgs);
+	}
+
+	bool FBuildCookRunLaunchExtension::CanBeCreated( ILauncherProfileRef InProfile, TSharedRef<FModel> InModel ) const
+	{
+		// cannot add a new BuildCookRun command to the basic launch profile (no technical reason - just want to avoid making it too complex)
+		return !InModel->IsBasicLaunchProfile(InProfile);
+	}
+}
+
+#undef LOCTEXT_NAMESPACE
diff --git a/Engine/Plugins/Developer/ProjectLauncher/Source/ProjectLauncher/Private/Extension/CustomUATCommandLaunchExtension.cpp b/Engine/Plugins/Developer/ProjectLauncher/Source/ProjectLauncher/Private/Extension/CustomUATCommandLaunchExtension.cpp
new file mode 100644
--- /dev/null
+++ b/Engine/Plugins/Developer/ProjectLauncher/Source/ProjectLauncher/Private/Extension/CustomUATCommandLaunchExtension.cpp
@@ -0,0 +1,51 @@
+// Copyright Epic Games, Inc. All Rights Reserved.
+
+#include "Extension/CustomUATCommandLaunchExtension.h"
+#include "Model/ProjectLauncherModel.h"
+
+#define LOCTEXT_NAMESPACE "FCustomUATCommandLaunchExtensionInstance"
+
+
+namespace ProjectLauncher
+{
+	FCustomUATCommandLaunchExtensionInstance::FCustomUATCommandLaunchExtensionInstance( FArgs& InArgs )
+		: FUATCommandLaunchExtensionInstanceBase(InArgs)
+	{
+	}
+
+	void FCustomUATCommandLaunchExtensionInstance::OnAdded()
+	{
+		// set the defaults, placing uat commands at the start by default
+		FModel::FUATCommandsDetail UATCommandsDetail; 
+		GetModel()->GetUATCommandsDetail(GetProfile(), UATCommandsDetail);
+		UATCommand->SetOrder( UATCommandsDetail.MinOrder - 1 ); 
+		UATCommand->SetDescription( GetExtension()->GetDisplayName().ToString() );
+
+		OnUATCommandAdded(UATCommand);
+
+		BroadcastPropertyChanged();
+	}
+	
+	void FCustomUATCommandLaunchExtensionInstance::OnRemoved()
+	{
+		OnUATCommandRemoved(UATCommand);
+		GetProfile()->RemoveUATCommand(UATCommand->GetInternalName());
+
+		BroadcastPropertyChanged();
+	}
+
+
+
+
+	void FCustomUATCommandLaunchExtension::GetExtensionsMenuEntry( FExtensionsMenuEntry& MenuEntry ) const
+	{
+		MenuEntry = FExtensionsMenuEntry::UATCommands;
+	}
+	
+	ILauncherProfileUATCommandRef FCustomUATCommandLaunchExtension::CreateUATCommand(ILauncherProfileRef InProfile, TSharedRef<FModel> InModel) const
+	{
+		return InProfile->FindOrAddUATCommand( nullptr, GetInternalName() );
+	}
+}
+
+#undef LOCTEXT_NAMESPACE
diff --git a/Engine/Plugins/Developer/ProjectLauncher/Source/ProjectLauncher/Private/Extension/LaunchExtension.cpp b/Engine/Plugins/Developer/ProjectLauncher/Source/ProjectLauncher/Private/Extension/LaunchExtension.cpp
--- a/Engine/Plugins/Developer/ProjectLauncher/Source/ProjectLauncher/Private/Extension/LaunchExtension.cpp
+++ b/Engine/Plugins/Developer/ProjectLauncher/Source/ProjectLauncher/Private/Extension/LaunchExtension.cpp
@@ -152,16 +152,7 @@ namespace ProjectLauncher
 				}
 
 				// allow for advanced command line customization
-				bool bKnownCommand = false;
-				const ILauncherProfileAutomatedTestPtr AutomatedTest = InPostProcessParameters.UATCommand->AsAutomatedTest();
-				if (AutomatedTest.IsValid())
-				{
-					bKnownCommand = ExtensionInstance->CustomizeAutomatedTestCommandLine( AutomatedTest.ToSharedRef(), InPostProcessParameters.UATCommandLine );
-				}
-				else
-				{
-					bKnownCommand = ExtensionInstance->CustomizeUATCommandLine( InPostProcessParameters.UATCommand, InPostProcessParameters.UATCommandLine );
-				}
+				bool bKnownCommand = ExtensionInstance->CustomizeUATCommandLine( InPostProcessParameters.UATCommand, InPostProcessParameters.UATCommandLine );
 
 				// don't skip this command if something knows about it
 				InPostProcessParameters.bShouldBeSkipped &= !bKnownCommand;
@@ -236,15 +227,42 @@ namespace ProjectLauncher
 			Extensions = GExtensions;
 		}
 
+		// cache all UAT commands in the profile, organised by type
+		TMap<FString,TArray<ILauncherProfileUATCommandRef>> UATCommandsByInternalName;
+		for (const ILauncherProfileUATCommandRef& UATCommand : InProfile->GetUATCommands() )
+		{
+			if (!UATCommandsByInternalName.Contains(UATCommand->GetUserTypeName()))
+			{
+				UATCommandsByInternalName.Add(UATCommand->GetUserTypeName());
+			}
+			UATCommandsByInternalName[UATCommand->GetUserTypeName()].Add(UATCommand);
+		}
+
+
 		// attempt to instantiate all extensions
 		for (TSharedPtr<FLaunchExtension> Extension : Extensions)
 		{
-			const bool* ValuePtr = InProfile->GetCustomBoolProperties().Find( FString::Printf(TEXT("%s.Instantiated"), Extension->GetInternalName() ) );
-			bool bEnabled = ValuePtr && (*ValuePtr);
+			if (Extension->IsUATCommandManager())
+			{
+				const TArray<ILauncherProfileUATCommandRef>* UATCommandsPtr = UATCommandsByInternalName.Find(Extension->GetInternalName());
+				if (UATCommandsPtr)
+				{
+					for (const ILauncherProfileUATCommandRef& UATCommand : (*UATCommandsPtr))
+					{
+						CreateExtensionInstance(Extension.ToSharedRef(), InProfile, InModel, InOwnerTreeData, UATCommand.ToSharedPtr());
+					}
+				}
 
-			if (Extension->IsAlwaysCreated(InProfile, InModel) || Extension->IsCreatedByDefault(InProfile, InModel) || bEnabled)
+			}
+			else
 			{
-				CreateExtensionInstance(Extension.ToSharedRef(), InProfile, InModel, InOwnerTreeData);
+				const bool* ValuePtr = InProfile->GetCustomBoolProperties().Find( FString::Printf(TEXT("%s.Instantiated"), Extension->GetInternalName() ) );
+				bool bEnabled = ValuePtr && (*ValuePtr);
+
+				if (Extension->IsAlwaysCreated(InProfile, InModel) || Extension->IsCreatedByDefault(InProfile, InModel) || bEnabled)
+				{
+					CreateExtensionInstance(Extension.ToSharedRef(), InProfile, InModel, InOwnerTreeData);
+				}
 			}
 		}
 
@@ -261,7 +279,7 @@ namespace ProjectLauncher
 		return TArray<TSharedPtr<FLaunchExtensionInstance>>();
 	}
 
-	TSharedPtr<FLaunchExtensionInstance> FLaunchExtension::CreateExtensionInstance( TSharedRef<FLaunchExtension> InExtension, ILauncherProfileRef InProfile, TSharedRef<FModel> InModel, FLaunchProfileTreeData& InOwnerTreeData )
+	TSharedPtr<FLaunchExtensionInstance> FLaunchExtension::CreateExtensionInstance( TSharedRef<FLaunchExtension> InExtension, ILauncherProfileRef InProfile, TSharedRef<FModel> InModel, FLaunchProfileTreeData& InOwnerTreeData, ILauncherProfileUATCommandPtr InUATCommand )
 	{
 		check(!GExtensionsDestroyed);
 
@@ -288,6 +306,7 @@ namespace ProjectLauncher
 			.Model = InModel,
 			.Extension = InExtension,
 			.OwnerTreeData = InOwnerTreeData,
+			.UATCommand = InUATCommand,
 		};
 		TSharedPtr<FLaunchExtensionInstance> ExtensionInstance = InExtension->CreateInstanceForProfile(Args);
 
@@ -296,7 +315,7 @@ namespace ProjectLauncher
 			ExtensionInstance->InternalInitialize();
 
 			// if this extension needs to be created dynamically, remember we have created it for next time and intialize it
-			if (!InExtension->IsAlwaysCreated(InProfile, InModel))
+			if (!InExtension->IsAlwaysCreated(InProfile, InModel) && !InExtension->IsUATCommandManager())
 			{
 				InProfile->GetCustomBoolProperties().Add( FString::Printf(TEXT("%s.Instantiated"), InExtension->GetInternalName() ), true );
 				ExtensionInstance->OnAdded();
@@ -317,25 +336,26 @@ namespace ProjectLauncher
 		return ExtensionInstance;
 	}
 
-	void FLaunchExtension::DestroyExtensionInstance( TSharedRef<FLaunchExtension> InExtension, ILauncherProfileRef InProfile, TSharedRef<FModel> InModel)
+	void FLaunchExtension::DestroyExtensionInstance( TSharedRef<FLaunchExtensionInstance> InExtensionInstance, ILauncherProfileRef InProfile, TSharedRef<FModel> InModel)
 	{
 		check(!GExtensionsDestroyed);
 
+		TSharedRef<FLaunchExtension> Extension = InExtensionInstance->GetExtension();
+
 		// can't destroy auto-instantiated extensions
-		check(!InExtension->IsAlwaysCreated(InProfile, InModel));
+		check(!Extension->IsAlwaysCreated(InProfile, InModel));
 
-		// remember it is no longer created
-		InProfile->GetCustomBoolProperties().Add( FString::Printf(TEXT("%s.Instantiated"), InExtension->GetInternalName() ), false );
+		if (!Extension->IsUATCommandManager())
+		{
+			InProfile->GetCustomBoolProperties().Add( FString::Printf(TEXT("%s.Instantiated"), Extension->GetInternalName() ), false );
+		}
 
 		// remove it from the global caches
-		TSharedPtr<FLaunchExtensionInstance> ExistingInstance = GetExtensionInstance(InExtension, InProfile);
-		if (ExistingInstance.IsValid())
+		InExtensionInstance->OnRemoved();
 		{
-			ExistingInstance->OnRemoved();
-
 			FScopeLock Lock(&GExtensionsCS);
-			GExtensionInstances.FindChecked(InProfile.ToSharedPtr()).Remove(ExistingInstance);
-			GInstancesPerExtension.FindChecked(InExtension.ToSharedPtr()).Remove(ExistingInstance);
+			GExtensionInstances.FindChecked(InProfile.ToSharedPtr()).Remove(InExtensionInstance);
+			GInstancesPerExtension.FindChecked(Extension.ToSharedPtr()).Remove(InExtensionInstance);
 		}
 	}
 
@@ -348,11 +368,13 @@ namespace ProjectLauncher
 	}
 
 
-	TSharedPtr<FLaunchExtensionInstance> FLaunchExtension::GetExtensionInstance(TSharedRef<FLaunchExtension> InExtension, ILauncherProfileRef InProfile)
+	TArray<TSharedRef<FLaunchExtensionInstance>> FLaunchExtension::GetExtensionInstances(TSharedRef<FLaunchExtension> InExtension, ILauncherProfileRef InProfile)
 	{
 		check(!GExtensionsDestroyed);
 
 		FScopeLock Lock(&GExtensionsCS);
+
+		TArray<TSharedRef<FLaunchExtensionInstance>> Result;
 		
 		const TArray<TSharedPtr<FLaunchExtensionInstance>>* InstancesForProfile = GExtensionInstances.Find(InProfile.ToSharedPtr());
 		const TArray<TSharedPtr<FLaunchExtensionInstance>>* InstancesForExtension = GInstancesPerExtension.Find(InExtension.ToSharedPtr());
@@ -362,12 +384,28 @@ namespace ProjectLauncher
 			{
 				if (InstancesForExtension->Contains(ExtensionInstance))
 				{
-					return ExtensionInstance;
+					Result.Add(ExtensionInstance.ToSharedRef());
 				}
 			}
 		}
 
-		return nullptr;
+		return MoveTemp(Result);
+	}
+
+	TSharedPtr<FLaunchExtensionInstance> FLaunchExtension::GetExtensionInstance(TSharedRef<FLaunchExtension> InExtension, ILauncherProfileRef InProfile)
+	{
+		if (InExtension->IsUATCommandManager())
+		{
+			return nullptr;
+		}
+
+		TArray<TSharedRef<FLaunchExtensionInstance>> ExtensionInstances = GetExtensionInstances(InExtension, InProfile);
+		if (ExtensionInstances.Num() == 0)
+		{
+			return nullptr;
+		}
+		check(ExtensionInstances.Num() == 1);
+		return ExtensionInstances[0];
 	}
 
 
@@ -749,14 +787,14 @@ namespace ProjectLauncher
 			case EConfig::PerProfile:
 			case EConfig::User_Common:
 			{
-				return FString::Printf(TEXT("%s.%s"), Extension->GetInternalName(), Name );
+				return FString::Printf(TEXT("%s.%s"), GetInternalName(), Name );
 			}
 			break;
 
 			case EConfig::User_PerProfile:
 			{
 				FGuid ProfileId = Profile->GetId();
-				return FString::Printf(TEXT("%s.%s.%s"), Extension->GetInternalName(), *ProfileId.ToString(), Name );
+				return FString::Printf(TEXT("%s.%s.%s"), GetInternalName(), *ProfileId.ToString(), Name );
 			}
 			break;
 		}
@@ -782,17 +820,7 @@ namespace ProjectLauncher
 
 	FLaunchProfileTreeNode& FLaunchExtensionInstance::AddDefaultHeading(ProjectLauncher::FLaunchProfileTreeData& ProfileTreeData) const
 	{
-		return AddDefaultHeading(ProfileTreeData, Extension->GetDisplayName());
-	}
-
-	FLaunchProfileTreeNode& FLaunchExtensionInstance::AddDefaultHeading(ProjectLauncher::FLaunchProfileTreeData& ProfileTreeData, FText OverrideDisplayName) const
-	{
-		return AddDefaultHeadingInternal(ProfileTreeData, OverrideDisplayName, GetHeadingSortOrder() );
-	}
-
-	FLaunchProfileTreeNode& FLaunchExtensionInstance::AddDefaultHeadingInternal(ProjectLauncher::FLaunchProfileTreeData& ProfileTreeData, FText DisplayName, int32 SortOrder) const
-	{
-		return ProfileTreeData.AddHeading(Extension->GetInternalName(), DisplayName, SortOrder );
+		return ProfileTreeData.AddHeading(GetInternalName(), Extension->GetDisplayName() );
 	}
 
 	bool FLaunchExtensionInstance::CustomizeUATCommandLine( const ILauncherProfileUATCommandRef& InUATCommand, FString& InOutCommandLine )
@@ -809,36 +837,10 @@ namespace ProjectLauncher
 		OnValidateProfile();
 	}
 
-
-
-
-
-
-
-
-
-
-
-
-
-
-
-
-
-
-
-
-
-
-
-
-
-
-
-
-
-
-
+	const TCHAR* FLaunchExtensionInstance::GetInternalName() const
+	{
+		return Extension->GetInternalName();
+	}
 }
 
 #undef LOCTEXT_NAMESPACE
diff --git a/Engine/Plugins/Developer/ProjectLauncher/Source/ProjectLauncher/Private/Extension/UATCommandLaunchExtension.cpp b/Engine/Plugins/Developer/ProjectLauncher/Source/ProjectLauncher/Private/Extension/UATCommandLaunchExtension.cpp
--- a/Engine/Plugins/Developer/ProjectLauncher/Source/ProjectLauncher/Private/Extension/UATCommandLaunchExtension.cpp
+++ b/Engine/Plugins/Developer/ProjectLauncher/Source/ProjectLauncher/Private/Extension/UATCommandLaunchExtension.cpp
@@ -7,104 +7,65 @@
 
 namespace ProjectLauncher
 {
-	FUATCommandLaunchExtensionInstance::FUATCommandLaunchExtensionInstance( FArgs& InArgs )
+	FUATCommandLaunchExtensionInstanceBase::FUATCommandLaunchExtensionInstanceBase( FArgs& InArgs )
 		: FLaunchExtensionInstance(InArgs)
+		, UATCommand(InArgs.UATCommand.ToSharedRef())
 	{
 	}
 
-	void FUATCommandLaunchExtensionInstance::InternalInitialize()
+	bool FUATCommandLaunchExtensionInstanceBase::ManagesUATCommand( const ILauncherProfileUATCommandRef& InUATCommand ) const
 	{
-		FLaunchExtensionInstance::InternalInitialize();
+		return UATCommand == InUATCommand;
 	}
 
-	void FUATCommandLaunchExtensionInstance::OnCustomValidation()
+	ILauncherProfileUATCommandPtr FUATCommandLaunchExtensionInstanceBase::GetUATCommand() const
 	{
-		ILauncherProfileUATCommandPtr UATCommand = GetUATCommand();
-		if (UATCommand.IsValid() && UATCommand->IsEnabled())
-		{
-			OnValidateProfile();
-		}
+		return UATCommand;
 	}
 
-	bool FUATCommandLaunchExtensionInstance::ManagesUATCommand( const ILauncherProfileUATCommandRef& UATCommand ) const
+	void FUATCommandLaunchExtensionInstanceBase::OnCustomValidation()
 	{
-		return UATCommand == GetUATCommand();
+		if (UATCommand->IsEnabled())
+		{
+			OnValidateProfile();
+		}
 	}
 
 
-
-	FLaunchProfileTreeNode& FUATCommandLaunchExtensionInstance::AddDefaultHeadingInternal(ProjectLauncher::FLaunchProfileTreeData& ProfileTreeData, FText DisplayName, int32 SortOrder) const
+	bool FUATCommandLaunchExtensionInstanceBase::CustomizeUATCommandLine( const ILauncherProfileUATCommandRef& InUATCommand, FString& InOutCommandLine )
 	{
-		ILauncherProfileUATCommandPtr UATCommand = GetUATCommand();
-		if (UATCommand.IsValid())
+		if (InUATCommand == UATCommand)
 		{
-			return ProfileTreeData.FindOrAddHeading(UATCommand.ToSharedRef(), true );
+			CustomizeUATCommandLine(InOutCommandLine);
+			return true;
 		}
 
-		return FLaunchExtensionInstance::AddDefaultHeadingInternal(ProfileTreeData, DisplayName, SortOrder);
+		return false;
 	}
 
-	void FUATCommandLaunchExtensionInstance::OnAdded()
-	{
-		bool bIsNew = false;
-		ILauncherProfileUATCommandPtr UATCommand = GetUATCommand();
-		if (!UATCommand.IsValid())
-		{
-			UATCommand = GetProfile()->FindOrAddUATCommand(*GetUATCommandInternalName());
-			bIsNew = true;
-		}
-		OnUATCommandAdded(UATCommand.ToSharedRef());
-
-		if (bIsNew)
-		{
-			int32 Order = GetUATCommandOrder();
-			UATCommand->SetOrder(Order);
-			UATCommand->SetDescription( GetExtension()->GetDisplayName().ToString() );
-		}
 
-		BroadcastPropertyChanged();
-	}
-	void FUATCommandLaunchExtensionInstance::OnRemoved()
+	FLaunchProfileTreeNode& FUATCommandLaunchExtensionInstanceBase::AddDefaultHeading(ProjectLauncher::FLaunchProfileTreeData& ProfileTreeData) const
 	{
-		ILauncherProfileUATCommandPtr ExistingUATCommand = GetUATCommand();
-		if (ExistingUATCommand.IsValid())
-		{
-			OnUATCommandRemoved(ExistingUATCommand.ToSharedRef());
-			GetProfile()->RemoveUATCommand(*GetUATCommandInternalName());
-		}
-		BroadcastPropertyChanged();
+		return ProfileTreeData.FindOrAddHeading(UATCommand, true);
 	}
 
-	bool FUATCommandLaunchExtensionInstance::IsUATCommandActive() const
+	const TCHAR* FUATCommandLaunchExtensionInstanceBase::GetInternalName() const
 	{
-		return GetUATCommand() != nullptr;
+		return UATCommand->GetInternalName();
 	}
 
-	ILauncherProfileUATCommandPtr FUATCommandLaunchExtensionInstance::GetUATCommand() const
-	{
-		return GetProfile()->GetUATCommand(*GetUATCommandInternalName());
-	}
 
 
-	bool FUATCommandLaunchExtensionInstance::CustomizeUATCommandLine( const ILauncherProfileUATCommandRef& InUATCommand, FString& InOutCommandLine )
+	TSharedPtr<ProjectLauncher::FLaunchExtensionInstance> FUATCommandLaunchExtensionBase::CreateUATCommandExtensionInstance(ILauncherProfileRef InProfile, TSharedRef<FModel> InModel, FLaunchProfileTreeData& InOwnerTreeData)
 	{
-		if (InUATCommand == GetUATCommand())
-		{
-			CustomizeUATCommandLine(InOutCommandLine);
-			return true;
-		}
+		ILauncherProfileUATCommandRef NewUATCommand = CreateUATCommand(InProfile, InModel);
+		TSharedPtr<FLaunchExtensionInstance> ExtensionInstance = FLaunchExtension::CreateExtensionInstance(AsShared(), InProfile, InModel, InOwnerTreeData, NewUATCommand.ToSharedPtr());
+		ExtensionInstance->OnAdded();
 
-		return false;
+		return ExtensionInstance;
 	}
 
 
-
-
-
-	void FUATCommandLaunchExtension::GetExtensionsMenuEntry( FExtensionsMenuEntry& MenuEntry ) const
-	{
-		MenuEntry = FExtensionsMenuEntry::UATCommands;
-	}
 }
 
 #undef LOCTEXT_NAMESPACE
diff --git a/Engine/Plugins/Developer/ProjectLauncher/Source/ProjectLauncher/Private/Model/ProjectLauncherModel.cpp b/Engine/Plugins/Developer/ProjectLauncher/Source/ProjectLauncher/Private/Model/ProjectLauncherModel.cpp
--- a/Engine/Plugins/Developer/ProjectLauncher/Source/ProjectLauncher/Private/Model/ProjectLauncherModel.cpp
+++ b/Engine/Plugins/Developer/ProjectLauncher/Source/ProjectLauncher/Private/Model/ProjectLauncherModel.cpp
@@ -836,17 +836,6 @@ PRAGMA_ENABLE_DEPRECATION_WARNINGS
 	}
 
 
-	ILauncherProfileBuildCookRunRef FModel::CreateBuildCookRun( const ILauncherProfileRef& Profile )
-	{
-		ILauncherProfileBuildCookRunRef BuildCookRun = Profile->FindOrAddBuildCookRunCommand(nullptr);
-
-		SetDefaults(Profile, BuildCookRun);
-
-		return BuildCookRun;
-	}
-
-
-
 	void FModel::ReadProjectSettingsFromConfig( FConfigCacheIni& InConfig, const FString& InProjectPath, FProjectSettings& OutResult )
 	{
 		// read build settings
@@ -1192,6 +1181,25 @@ PRAGMA_ENABLE_DEPRECATION_WARNINGS
 		return GetAndCacheMapPaths(FString());
 	}
 
+	void FModel::GetUATCommandsDetail( const ILauncherProfileRef& Profile, FUATCommandsDetail& OutUATCommandsDetail)
+	{
+		const TArray<ILauncherProfileUATCommandRef>& UATCommands = Profile->GetUATCommands();
+		OutUATCommandsDetail.NumUATCommands = UATCommands.Num();
+		if (OutUATCommandsDetail.NumUATCommands > 0)
+		{
+			OutUATCommandsDetail.MaxOrder = INT32_MIN;
+			OutUATCommandsDetail.MinOrder = INT32_MAX;
+
+			for (const ILauncherProfileUATCommandRef& UATCommand : UATCommands)
+			{
+				int32 Order = UATCommand->GetOrder();
+
+				OutUATCommandsDetail.MaxOrder = FMath::Max(OutUATCommandsDetail.MaxOrder, Order);
+				OutUATCommandsDetail.MinOrder = FMath::Min(OutUATCommandsDetail.MinOrder, Order);
+			}
+		}
+	}
+
 
 
 
diff --git a/Engine/Plugins/Developer/ProjectLauncher/Source/ProjectLauncher/Private/ProjectLauncherModule.cpp b/Engine/Plugins/Developer/ProjectLauncher/Source/ProjectLauncher/Private/ProjectLauncherModule.cpp
--- a/Engine/Plugins/Developer/ProjectLauncher/Source/ProjectLauncher/Private/ProjectLauncherModule.cpp
+++ b/Engine/Plugins/Developer/ProjectLauncher/Source/ProjectLauncher/Private/ProjectLauncherModule.cpp
@@ -13,6 +13,7 @@
 #include "ProfileTree/ILaunchProfileTreeBuilder.h"
 #include "ProfileTree/BasicProfileTreeBuilder.h"
 #include "ProfileTree/CustomProfileTreeBuilder.h"
+#include "Extension/BuildCookRunLaunchExtension.h"
 #include "Model/ProjectLauncherModel.h"
 
 #define LOCTEXT_NAMESPACE "FProjectLauncherModule"
@@ -42,6 +43,9 @@ public:
 		RegisterTreeBuilder( MakeShared<ProjectLauncher::FBasicProfileTreeBuilderFactory>() );
 		RegisterTreeBuilder( MakeShared<ProjectLauncher::FCustomProfileTreeBuilderFactory>() );
 
+		// register built-in extensions
+		RegisterExtension( MakeShared<ProjectLauncher::FBuildCookRunLaunchExtension>() );
+
 		// register custom styles & other UI elements
 		FProjectLauncherStyle::Initialize();
 		if (!FDialogCommands::IsRegistered())
diff --git a/Engine/Plugins/Developer/ProjectLauncher/Source/ProjectLauncher/Private/Widgets/Profiles/SCustomLaunchCustomProfileEditor.cpp b/Engine/Plugins/Developer/ProjectLauncher/Source/ProjectLauncher/Private/Widgets/Profiles/SCustomLaunchCustomProfileEditor.cpp
--- a/Engine/Plugins/Developer/ProjectLauncher/Source/ProjectLauncher/Private/Widgets/Profiles/SCustomLaunchCustomProfileEditor.cpp
+++ b/Engine/Plugins/Developer/ProjectLauncher/Source/ProjectLauncher/Private/Widgets/Profiles/SCustomLaunchCustomProfileEditor.cpp
@@ -7,6 +7,7 @@
 #include "Styling/ProjectLauncherStyle.h"
 #include "Extension/LaunchExtension.h"
 #include "Extension/BuildCookRunCommandExtension.h"
+#include "Extension/UATCommandLaunchExtension.h"
 #include "Framework/Commands/UIAction.h"
 #include "Framework/MultiBox/MultiBoxBuilder.h"
 #include "SlateOptMacros.h"
@@ -929,10 +930,15 @@ void SCustomLaunchCustomProfileEditor::ApplyCollapsedRecusive( TSet<FString>& Co
 }
 
 
-TSharedRef<SWidget> SCustomLaunchCustomProfileEditor::MakeExtensionsMenu()
+void SCustomLaunchCustomProfileEditor::MakeMenu( FMenuBuilder& MenuBuilder, EMenuType MenuType, ILauncherProfileBuildCookRunPtr BuildCookRun )
 {
 	ProjectLauncher::FLaunchProfileTreeDataRef TreeData = TreeBuilder->GetProfileTree();
 
+	if (Model->CanUseSimplifiedLayout(CurrentProfile.ToSharedRef()) && !BuildCookRun.IsValid())
+	{
+		BuildCookRun = CurrentProfile->GetFirstBuildCookRun();
+	}
+
 	// ...build extension submenu structure...
 
 	struct FExtensionMenuSubMenu
@@ -961,13 +967,13 @@ TSharedRef<SWidget> SCustomLaunchCustomProfileEditor::MakeExtensionsMenu()
 			continue;
 		}
 
-		// don't show additional UAT commands or automated tests in the Extensions menu - save them for the Tasks menu
-		if (Extension->AsUATCommandLaunchExtension() || Extension->AsAutomatedTestLaunchExtension())
+		// filter based on menu type (this prevents empty submenus)
+		bool bIsUATCommand = Extension->IsUATCommandManager();
+		if (bIsUATCommand != (MenuType == EMenuType::UATCommands))
 		{
 			continue;
 		}
 
-
 		// get section
 		FString SectionName = ExtensionMenuEntry.SectionName;
 		FText SectionDisplayName = ExtensionMenuEntry.SectionDisplayName;
@@ -1012,13 +1018,6 @@ TSharedRef<SWidget> SCustomLaunchCustomProfileEditor::MakeExtensionsMenu()
 
 	// ...construct the extensions menu...
 
-	const bool bShouldCloseWindowAfterMenuSelection = true;
-	const bool bCloseSelfOnly = false;
-	const bool bSearchable = false;
-	FMenuBuilder MenuBuilder(bShouldCloseWindowAfterMenuSelection, nullptr, nullptr, bCloseSelfOnly, &FCoreStyle::Get(), bSearchable );
-	
-
-	bool bIncludeBuildCookRunExtensions = Model->CanUseSimplifiedLayout(CurrentProfile.ToSharedRef());
 	for (const auto& SectionItr : Sections)
 	{
 		// start the section
@@ -1027,22 +1026,24 @@ TSharedRef<SWidget> SCustomLaunchCustomProfileEditor::MakeExtensionsMenu()
 		// add all extensions
 		for (TSharedRef<ProjectLauncher::FLaunchExtension> Extension : SectionItr.Value.Extensions)
 		{
-			MakeExtensionMenu(MenuBuilder, Extension, bIncludeBuildCookRunExtensions);
+			MakeExtensionMenu(MenuBuilder, MenuType, Extension, BuildCookRun);
 		}
 
 		// add all submenus
 		for (const auto& SubMenuItr : SectionItr.Value.SubMenus)
 		{
-			auto MakeExtensionSubMenu = [this, bIncludeBuildCookRunExtensions](FMenuBuilder& MenuBuilder, TSharedPtr<FExtensionMenuSubMenu> SubMenu)
+			auto MakeExtensionSubMenu = [this, MenuType, BuildCookRun](FMenuBuilder& MenuBuilder, TSharedPtr<FExtensionMenuSubMenu> SubMenu)
 			{
-				MenuBuilder.SetSearchable(bSearchable);
+				MenuBuilder.SetSearchable(false);
 				for (TSharedRef<ProjectLauncher::FLaunchExtension> Extension : SubMenu->Extensions)
 				{
-					MakeExtensionMenu(MenuBuilder, Extension, bIncludeBuildCookRunExtensions);
+					MakeExtensionMenu(MenuBuilder, MenuType, Extension, BuildCookRun);
 				}
 			};
 
 			const bool bInOpenSubMenuOnClick = false;
+			const bool bShouldCloseWindowAfterMenuSelection = true;
+
 			MenuBuilder.AddSubMenu(
 				SubMenuItr.Value->DisplayName,
 				FText::GetEmpty(),
@@ -1057,21 +1058,38 @@ TSharedRef<SWidget> SCustomLaunchCustomProfileEditor::MakeExtensionsMenu()
 		MenuBuilder.EndSection();
 	}
 
-	return MenuBuilder.MakeWidget();
 }
 
-void SCustomLaunchCustomProfileEditor::MakeExtensionMenu( FMenuBuilder& MenuBuilder, TSharedRef<ProjectLauncher::FLaunchExtension> Extension, bool bIncludeBuildCookRunExtensions )
+void SCustomLaunchCustomProfileEditor::MakeExtensionMenu( FMenuBuilder& MenuBuilder, EMenuType MenuType, TSharedRef<ProjectLauncher::FLaunchExtension> Extension, ILauncherProfileBuildCookRunPtr BuildCookRun )
 {
 	TSharedPtr<ProjectLauncher::FLaunchExtensionInstance> ExtensionInstance = ProjectLauncher::FLaunchExtension::GetExtensionInstance(Extension, CurrentProfile.ToSharedRef());
 
-	if (!Extension->IsAlwaysCreated(CurrentProfile.ToSharedRef(), Model.ToSharedRef()))
+	if (MenuType == EMenuType::UATCommands && Extension->IsUATCommandManager() && Extension->CanBeCreated(CurrentProfile.ToSharedRef(), Model.ToSharedRef()) )
+	{
+		MenuBuilder.AddMenuEntry(
+			FText::Format(LOCTEXT("NewUATCmdLabel", "Add New {0}"), Extension->GetDisplayName()),
+			FText::GetEmpty(),
+			FSlateIcon(Extension->GetIcon()),
+			FUIAction(
+				FExecuteAction::CreateLambda( [this, Extension]()
+				{
+					TSharedPtr<ProjectLauncher::FLaunchExtensionInstance> Instance = Extension->AsUATCommandFactory()->CreateUATCommandExtensionInstance(CurrentProfile.ToSharedRef(), Model.ToSharedRef(), TreeBuilder->GetProfileTree().Get());
+					if (Instance.IsValid())
+					{
+						TreeBuilder->GetProfileTree()->RequestFullTreeRebuild();
+					}
+				})
+			),
+			NAME_None);
+	}
+	else if (MenuType == EMenuType::ProfileExtensions && !Extension->IsUATCommandManager() && !Extension->IsAlwaysCreated(CurrentProfile.ToSharedRef(), Model.ToSharedRef()))
 	{
 		auto OnToggleExtension = [this, Extension, ExtensionInstance]()
 		{
 			if (ExtensionInstance.IsValid())
 			{
 				// defer the delete until the next tick
-				PendingExtensionDeletes.AddUnique(Extension);
+				PendingExtensionDeletes.AddUnique(ExtensionInstance.ToSharedRef());
 			}
 			else
 			{
@@ -1128,76 +1146,54 @@ void SCustomLaunchCustomProfileEditor::MakeExtensionMenu( FMenuBuilder& MenuBuil
 
 	if (ExtensionInstance)
 	{
-		if (bIncludeBuildCookRunExtensions && CurrentProfile->GetFirstBuildCookRun().IsValid())
+		if (BuildCookRun.IsValid())
 		{
 			if (ProjectLauncher::IBuildCookRunExtensionFactory* BuildCookRunFactory = ExtensionInstance->AsBuildCookRunFactory())
 			{
-				BuildCookRunFactory->CustomizeBuildCookRunExtensionMenu(MenuBuilder, CurrentProfile->GetFirstBuildCookRun().ToSharedRef());
+				BuildCookRunFactory->CustomizeBuildCookRunExtensionMenu(MenuBuilder, BuildCookRun.ToSharedRef());
 			}
 		}
 
-		ExtensionInstance->MakeCustomExtensionSubmenu(MenuBuilder);
+		if (MenuType == EMenuType::ProfileExtensions)
+		{
+			ExtensionInstance->MakeCustomExtensionSubmenu(MenuBuilder);
+		}
 	}
 
-	Extension->MakeCustomExtensionSubmenu(MenuBuilder, CurrentProfile.ToSharedRef(), Model.ToSharedRef());
+	if (MenuType == EMenuType::ProfileExtensions)
+	{
+		Extension->MakeCustomExtensionSubmenu(MenuBuilder, CurrentProfile.ToSharedRef(), Model.ToSharedRef());
+	}
 };
 
 
 
 
 
+TSharedRef<SWidget> SCustomLaunchCustomProfileEditor::MakeExtensionsMenu()
+{
+	// ...construct the extensions menu...
+
+	const bool bShouldCloseWindowAfterMenuSelection = true;
+	const bool bCloseSelfOnly = false;
+	const bool bSearchable = false;
+	FMenuBuilder MenuBuilder(bShouldCloseWindowAfterMenuSelection, nullptr, nullptr, bCloseSelfOnly, &FCoreStyle::Get(), bSearchable );
+
+	MakeMenu(MenuBuilder, EMenuType::ProfileExtensions);
+
+	return MenuBuilder.MakeWidget();
+}
+
 
 TSharedRef<SWidget> SCustomLaunchCustomProfileEditor::MakeCommandsMenu()
 {
-	ProjectLauncher::FLaunchProfileTreeDataRef TreeData = TreeBuilder->GetProfileTree();
-
 	const bool bShouldCloseWindowAfterMenuSelection = true;
 	const bool bCloseSelfOnly = true;
 	const bool bSearchable = false;
 	const bool bInOpenSubMenuOnClick = false;
 	FMenuBuilder MenuBuilder(bShouldCloseWindowAfterMenuSelection, nullptr, nullptr, bCloseSelfOnly, &FCoreStyle::Get(), bSearchable );
 
-	// cannot add new BuildCookRuns to the basic launch profile
-	if (!Model->IsBasicLaunchProfile(CurrentProfile))
-	{
-		MenuBuilder.AddMenuEntry(
-			LOCTEXT("NewBuildCookRunCmdLabel", "Add New Build, Cook & Run"),
-			FText::GetEmpty(),
-			FSlateIcon(FProjectLauncherStyle::GetStyleSetName(), "Icons.Task.Launch"),
-			FUIAction(
-				FExecuteAction::CreateLambda( [this, TreeData]()
-				{
-					ILauncherProfileBuildCookRunRef NewBuildCookRun = TreeData->Model->CreateBuildCookRun(CurrentProfile.ToSharedRef());
-					NewBuildCookRun->SetDescription(TEXT("New Task (Build, Cook & Run)"));
-					
-					TreeData->RequestFullTreeRebuild();
-				})
-			),
-			NAME_None);
-	}
-
-	// enumerate extensions & automated tests & add them here too. 
-	// @todo: honor FExtensionsMenuEntry structure here
-	// @todo: UAT commands and AutomatedTests should also be add/remove like BuildCookRun
-	MenuBuilder.BeginSection(NAME_None, LOCTEXT("UATCommandsGroup", "Add Build Tasks"));
-	for (const TSharedPtr<ProjectLauncher::FLaunchExtension> Extension : ProjectLauncher::FLaunchExtension::GetExtensions())
-	{
-		if (Extension->AsUATCommandLaunchExtension() != nullptr)
-		{
-			MakeExtensionMenu(MenuBuilder, Extension.ToSharedRef(), false);
-		}
-	}
-	MenuBuilder.EndSection();
-
-	MenuBuilder.BeginSection(NAME_None, LOCTEXT("AutomatedTestsGroup", "Add Automated Tests"));
-	for (const TSharedPtr<ProjectLauncher::FLaunchExtension> Extension : ProjectLauncher::FLaunchExtension::GetExtensions())
-	{
-		if (Extension->AsAutomatedTestLaunchExtension() != nullptr)
-		{
-			MakeExtensionMenu(MenuBuilder, Extension.ToSharedRef(), false);
-		}
-	}
-	MenuBuilder.EndSection();
+	MakeMenu( MenuBuilder, EMenuType::UATCommands );
 
 	return MenuBuilder.MakeWidget();
 }
@@ -1255,7 +1251,7 @@ void SCustomLaunchCustomProfileEditor::MakeUATCommandSubMenu( FMenuBuilder& Menu
 					if (ExtensionInstance.IsValid())
 					{
 						// defer the delete until the next tick
-						PendingExtensionDeletes.AddUnique(ExtensionInstance->GetExtension());
+						PendingExtensionDeletes.AddUnique(ExtensionInstance.ToSharedRef());
 					}
 					else
 					{
@@ -1271,26 +1267,10 @@ void SCustomLaunchCustomProfileEditor::MakeUATCommandSubMenu( FMenuBuilder& Menu
 	ILauncherProfileBuildCookRunPtr BuildCookRun = UATCommand->AsBuildCookRun();
 	if (BuildCookRun.IsValid())
 	{
-		MenuBuilder.BeginSection(NAME_None, LOCTEXT("ExtensionsGroupLabel", "Extensions"));
-		MakeBuildCookRunExtensionsSubMenu(MenuBuilder, BuildCookRun.ToSharedRef());
-		MenuBuilder.EndSection();
-	}
-}
-
-void SCustomLaunchCustomProfileEditor::MakeBuildCookRunExtensionsSubMenu( FMenuBuilder& MenuBuilder, ILauncherProfileBuildCookRunRef BuildCookRun ) const
-{
-	ProjectLauncher::FLaunchProfileTreeDataRef TreeData = TreeBuilder->GetProfileTree();
-
-	for (TSharedPtr<ProjectLauncher::FLaunchExtensionInstance> ExtensionInstance : TreeData->ExtensionInstances)
-	{
-		if (ProjectLauncher::IBuildCookRunExtensionFactory* BuildCookRunFactory = ExtensionInstance->AsBuildCookRunFactory())
-		{
-			BuildCookRunFactory->CustomizeBuildCookRunExtensionMenu(MenuBuilder, BuildCookRun);
-		}
+		MakeMenu(MenuBuilder, EMenuType::BuildCookRunExtensions, BuildCookRun);
 	}
 }
 
-
 TSharedRef<SWidget> SCustomLaunchCustomProfileEditor::MakeUATCommandSubMenuWidget(ILauncherProfileUATCommandRef UATCommand)
 {
 	const bool bShouldCloseWindowAfterMenuSelection = true;
@@ -1367,9 +1347,9 @@ void SCustomLaunchCustomProfileEditor::Tick(const FGeometry& AllottedGeometry, c
 			ExecuteOnGameThread(UE_SOURCE_LOCATION,[this, PendingDeletes = MoveTemp(PendingExtensionDeletes)]()
 			{
 				bool bProfileDirty = false;
-				for (const TSharedRef<ProjectLauncher::FLaunchExtension>& Extension : PendingDeletes)
+				for (const TSharedRef<ProjectLauncher::FLaunchExtensionInstance>& ExtensionInstance : PendingDeletes)
 				{
-					ProjectLauncher::FLaunchExtension::DestroyExtensionInstance(Extension, CurrentProfile.ToSharedRef(), Model.ToSharedRef());
+					ProjectLauncher::FLaunchExtension::DestroyExtensionInstance(ExtensionInstance, CurrentProfile.ToSharedRef(), Model.ToSharedRef());
 					bProfileDirty = true;
 				}
 
diff --git a/Engine/Plugins/Developer/ProjectLauncher/Source/ProjectLauncher/Private/Widgets/Profiles/SCustomLaunchCustomProfileEditor.h b/Engine/Plugins/Developer/ProjectLauncher/Source/ProjectLauncher/Private/Widgets/Profiles/SCustomLaunchCustomProfileEditor.h
--- a/Engine/Plugins/Developer/ProjectLauncher/Source/ProjectLauncher/Private/Widgets/Profiles/SCustomLaunchCustomProfileEditor.h
+++ b/Engine/Plugins/Developer/ProjectLauncher/Source/ProjectLauncher/Private/Widgets/Profiles/SCustomLaunchCustomProfileEditor.h
@@ -58,9 +58,16 @@ private:
 	void ApplyCollapsedRecusive( TSet<FString>& Collapsed, const FString& NodeName, ProjectLauncher::FLaunchProfileTreeNodePtr Node) const;
 	void ReorderNodes(ProjectLauncher::FLaunchProfileTreeNodePtr SourceNode, ProjectLauncher::FLaunchProfileTreeNodePtr TargetNode, EItemDropZone DropZone);
 
+	enum class EMenuType
+	{
+		ProfileExtensions,
+		BuildCookRunExtensions,
+		UATCommands,
+	};
+	void MakeMenu( FMenuBuilder& MenuBuilder, EMenuType MenuType, ILauncherProfileBuildCookRunPtr BuildCookRun = nullptr );
+
 	void MakeUATCommandSubMenu( FMenuBuilder& MenuBuilder, ILauncherProfileUATCommandRef UATCommand );
-	void MakeBuildCookRunExtensionsSubMenu( FMenuBuilder& MenuBuilder, ILauncherProfileBuildCookRunRef BuildCookRun ) const;
-	void MakeExtensionMenu( FMenuBuilder& MenuBuilder, TSharedRef<ProjectLauncher::FLaunchExtension> Extension, bool bIncludeBuildCookRunExtensions );
+	void MakeExtensionMenu( FMenuBuilder& MenuBuilder, EMenuType MenuType, TSharedRef<ProjectLauncher::FLaunchExtension> Extension, ILauncherProfileBuildCookRunPtr BuildCookRun = nullptr );
 	TSharedRef<SWidget> MakeUATCommandSubMenuWidget(ILauncherProfileUATCommandRef UATCommand);
 
 	TSharedPtr<ProjectLauncher::FLaunchExtensionInstance> GetExtensionForUATCommand(ILauncherProfileUATCommandRef UATCommand) const;
@@ -70,7 +77,7 @@ private:
 	TSharedPtr<ProjectLauncher::FModel> Model;
 	TSharedPtr<ProjectLauncher::ILaunchProfileTreeBuilder> TreeBuilder;
 	TSharedPtr<SLaunchProfileTreeView> TreeView;
-	TArray<TSharedRef<ProjectLauncher::FLaunchExtension>> PendingExtensionDeletes;
+	TArray<TSharedRef<ProjectLauncher::FLaunchExtensionInstance>> PendingExtensionDeletes;
 
 	float SplitterPos = 0.6f;
 	
diff --git a/Engine/Plugins/Developer/ProjectLauncher/Source/ProjectLauncher/Public/Extension/AutomatedTestLaunchExtension.h b/Engine/Plugins/Developer/ProjectLauncher/Source/ProjectLauncher/Public/Extension/AutomatedTestLaunchExtension.h
--- a/Engine/Plugins/Developer/ProjectLauncher/Source/ProjectLauncher/Public/Extension/AutomatedTestLaunchExtension.h
+++ b/Engine/Plugins/Developer/ProjectLauncher/Source/ProjectLauncher/Public/Extension/AutomatedTestLaunchExtension.h
@@ -2,7 +2,7 @@
 
 #pragma once
 
-#include "Extension/LaunchExtension.h"
+#include "Extension/UATCommandLaunchExtension.h"
 
 #define UE_API PROJECTLAUNCHER_API
 
@@ -13,74 +13,40 @@ namespace ProjectLauncher
 	 * Creates an item in the extension menu that will toggle the test on and off
 	 * note: It is not necessary for an automated test extensions to use this class
 	 */
-	class FAutomatedTestLaunchExtensionInstance : public FLaunchExtensionInstance
+	class FAutomatedTestLaunchExtensionInstance : public FUATCommandLaunchExtensionInstanceBase
 	{
 	public:
 		UE_API FAutomatedTestLaunchExtensionInstance( FArgs& InArgs );
 
-		/**
-		 * Determine if this automated test is currently active
-		 */
-		UE_API bool IsTestActive() const;
-
-		UE_API virtual bool ManagesUATCommand( const ILauncherProfileUATCommandRef& UATCommand ) const override;
-
 	protected:
 		UE_API virtual void OnAdded() override;
 		UE_API virtual void OnRemoved() override;
 
-		UE_API bool CustomizeAutomatedTestCommandLine( const ILauncherProfileAutomatedTestRef& InAutomatedTest, FString& InOutCommandLine ) override;
-		virtual bool HasCustomExtensionMenu() const { return true; }
-
-		/**
-		 * Advanced hook to allow any advanced modification of the command line when this extension's automated test is launched
-		 */
-		virtual void CustomizeAutomatedTestCommandLine( FString& InOutCommandLine ) {}
-
 		/**
 		 * Hook to allow any modification of the automated test & profile once it has been created. At least the test name should be specified
 		 */
-		virtual void OnTestAdded( ILauncherProfileAutomatedTestRef AutomatedTest ) = 0;
+		virtual void OnTestAdded( ILauncherProfileAutomatedTestRef InAutomatedTest ) = 0;
 
 		/**
 		 * Advanced hook to allow any modification of the profile just before the automated test is removed
 		 */
-		virtual void OnTestRemoved( ILauncherProfileAutomatedTestRef AutomatedTest ) {}
+		virtual void OnTestRemoved( ILauncherProfileAutomatedTestRef InAutomatedTest ) {}
 
 		/**
 		 * Access this extension's automated test if it has been created
 		 */
 		UE_API ILauncherProfileAutomatedTestPtr GetTest() const;
-
-		/**
-		 * Internal name for the automated test. Must be globally unique, so recommended to include the extension name
-		 */
-		virtual const FString GetTestInternalName() const = 0;
-
-		/**
-		 * Determine the ordering for this test. lower number = run earlier
-		 */
-		virtual int32 GetTestOrder() const { return INT32_MAX; }
-
-		/**
-		 * Internal intitialization function, called immediately after construction - don't call directly
-		 */
-		UE_API virtual void InternalInitialize() override;
-
-	private:
-		UE_API virtual void OnCustomValidation() override;
-		virtual int32 GetHeadingSortOrder() const override { return GetTestOrder(); }
-		UE_API virtual FLaunchProfileTreeNode& AddDefaultHeadingInternal(FLaunchProfileTreeData& ProfileTreeData, FText DisplayName, int32 SortOrder) const override;
 	};
 
 
-	class FAutomatedTestLaunchExtension : public FLaunchExtension
+	class FAutomatedTestLaunchExtension : public FUATCommandLaunchExtensionBase
 	{
 	public:
 		UE_API virtual void GetExtensionsMenuEntry( FExtensionsMenuEntry& MenuEntry ) const override;
 		UE_API virtual FSlateIcon GetIcon() const override;
 
-		virtual class FAutomatedTestLaunchExtension* AsAutomatedTestLaunchExtension() override { return this; }
+	protected:
+		UE_API virtual ILauncherProfileUATCommandRef CreateUATCommand(ILauncherProfileRef InProfile, TSharedRef<FModel> InModel) const override;
 	};
 };
 
diff --git a/Engine/Plugins/Developer/ProjectLauncher/Source/ProjectLauncher/Public/Extension/BuildCookRunCommandExtension.h b/Engine/Plugins/Developer/ProjectLauncher/Source/ProjectLauncher/Public/Extension/BuildCookRunCommandExtension.h
--- a/Engine/Plugins/Developer/ProjectLauncher/Source/ProjectLauncher/Public/Extension/BuildCookRunCommandExtension.h
+++ b/Engine/Plugins/Developer/ProjectLauncher/Source/ProjectLauncher/Public/Extension/BuildCookRunCommandExtension.h
@@ -99,7 +99,6 @@ namespace ProjectLauncher
 
 	private:
 		void OnUATCommandValidation( ILauncherProfileUATCommandRef UATCommand );
-		virtual int32 GetHeadingSortOrder() const { return 0; }
 
 		ILauncherProfileRef Profile;
 		ILauncherProfileBuildCookRunRef BuildCookRun;
@@ -206,7 +205,6 @@ namespace ProjectLauncher
 	{
 	public:
 		virtual bool IsAlwaysCreated( ILauncherProfileRef InProfile, TSharedRef<FModel> InModel ) const override { return true; }
-		virtual class FBuildCookRunCommandExtension* AsBuildCookRunCommandExtension() override { return this; }
 	};
 
 
diff --git a/Engine/Plugins/Developer/ProjectLauncher/Source/ProjectLauncher/Public/Extension/BuildCookRunLaunchExtension.h b/Engine/Plugins/Developer/ProjectLauncher/Source/ProjectLauncher/Public/Extension/BuildCookRunLaunchExtension.h
new file mode 100644
--- /dev/null
+++ b/Engine/Plugins/Developer/ProjectLauncher/Source/ProjectLauncher/Public/Extension/BuildCookRunLaunchExtension.h
@@ -0,0 +1,35 @@
+// Copyright Epic Games, Inc. All Rights Reserved.
+
+#pragma once
+
+#include "Extension/UATCommandLaunchExtension.h"
+
+#define UE_API PROJECTLAUNCHER_API
+
+namespace ProjectLauncher
+{
+	class FBuildCookRunLaunchExtensionInstance : public FUATCommandLaunchExtensionInstanceBase
+	{
+	public:
+		UE_API FBuildCookRunLaunchExtensionInstance( FArgs& InArgs );
+	};
+
+	class FBuildCookRunLaunchExtension : public FUATCommandLaunchExtensionBase
+	{
+	public:
+		UE_API virtual void GetExtensionsMenuEntry( FExtensionsMenuEntry& MenuEntry ) const override;
+		UE_API virtual FSlateIcon GetIcon() const override;
+		UE_API virtual const TCHAR* GetInternalName() const override;
+		UE_API virtual FText GetDisplayName() const override;
+
+		UE_API virtual TSharedPtr<ProjectLauncher::FLaunchExtensionInstance> CreateInstanceForProfile( ProjectLauncher::FLaunchExtensionInstance::FArgs& InArgs ) override;
+		UE_API virtual bool CanBeCreated( ILauncherProfileRef InProfile, TSharedRef<FModel> InModel ) const override;
+
+	protected:
+		UE_API virtual ILauncherProfileUATCommandRef CreateUATCommand(ILauncherProfileRef InProfile, TSharedRef<FModel> InModel) const override;
+	};
+
+
+};
+
+#undef UE_API
diff --git a/Engine/Plugins/Developer/ProjectLauncher/Source/ProjectLauncher/Public/Extension/CustomUATCommandLaunchExtension.h b/Engine/Plugins/Developer/ProjectLauncher/Source/ProjectLauncher/Public/Extension/CustomUATCommandLaunchExtension.h
new file mode 100644
--- /dev/null
+++ b/Engine/Plugins/Developer/ProjectLauncher/Source/ProjectLauncher/Public/Extension/CustomUATCommandLaunchExtension.h
@@ -0,0 +1,47 @@
+// Copyright Epic Games, Inc. All Rights Reserved.
+
+#pragma once
+
+#include "Extension/UATCommandLaunchExtension.h"
+
+#define UE_API PROJECTLAUNCHER_API
+
+namespace ProjectLauncher
+{
+	/**
+	 * Base class for an custom UAT command launch extension
+	 */
+	class FCustomUATCommandLaunchExtensionInstance : public FUATCommandLaunchExtensionInstanceBase
+	{
+	public:
+		UE_API FCustomUATCommandLaunchExtensionInstance( FArgs& InArgs );
+
+	protected:
+		UE_API virtual void OnAdded() override;
+		UE_API virtual void OnRemoved() override;
+
+		/**
+		 * Hook to allow any modification of the UAT Command & profile once it has been created. At least the UAT Command name should be specified
+		 */
+		virtual void OnUATCommandAdded( ILauncherProfileUATCommandRef InUATCommand ) = 0;
+
+		/**
+		 * Advanced hook to allow any modification of the profile just before the UAT Command is removed
+		 */
+		virtual void OnUATCommandRemoved( ILauncherProfileUATCommandRef InUATCommand ) {}
+	};
+
+
+
+	class FCustomUATCommandLaunchExtension : public FUATCommandLaunchExtensionBase
+	{
+	public:
+		UE_API virtual void GetExtensionsMenuEntry( FExtensionsMenuEntry& MenuEntry ) const override;
+
+	protected:
+		UE_API virtual ILauncherProfileUATCommandRef CreateUATCommand(ILauncherProfileRef InProfile, TSharedRef<FModel> InModel) const override;
+	};
+
+};
+
+#undef UE_API
diff --git a/Engine/Plugins/Developer/ProjectLauncher/Source/ProjectLauncher/Public/Extension/LaunchExtension.h b/Engine/Plugins/Developer/ProjectLauncher/Source/ProjectLauncher/Public/Extension/LaunchExtension.h
--- a/Engine/Plugins/Developer/ProjectLauncher/Source/ProjectLauncher/Public/Extension/LaunchExtension.h
+++ b/Engine/Plugins/Developer/ProjectLauncher/Source/ProjectLauncher/Public/Extension/LaunchExtension.h
@@ -38,6 +38,7 @@ namespace ProjectLauncher
 			TSharedRef<FModel> Model;
 			TSharedRef<FLaunchExtension> Extension;
 			FLaunchProfileTreeData& OwnerTreeData;
+			ILauncherProfileUATCommandPtr UATCommand;
 		};
 
 		UE_API FLaunchExtensionInstance( FArgs& InArgs );
@@ -87,11 +88,6 @@ namespace ProjectLauncher
 		 */
 		virtual void CustomizeTargetLaunchCommandLine( FString& InOutCommandLine, EBuildTargetType InBuildTargetType ) {};
 
-		/**
-		 * Advanced hook to allow any advanced modification of the command line when an automated test is launched
-		 */
-		virtual bool CustomizeAutomatedTestCommandLine( const ILauncherProfileAutomatedTestRef& InAutomatedTest, FString& InOutCommandLine ) { return false; }
-
 		/**
 		 * Advanced hook to allow any advanced modification of the command line when a UAT command test is launched
 		 */
@@ -115,12 +111,7 @@ namespace ProjectLauncher
 		/**
 		 * Create the default heading for this extension
 		 */
-		UE_API FLaunchProfileTreeNode& AddDefaultHeading(ProjectLauncher::FLaunchProfileTreeData& ProfileTreeData) const;
-
-		/**
-		 * Create the default heading for this extension, with a bespoke display name
-		 */
-		UE_API FLaunchProfileTreeNode& AddDefaultHeading(ProjectLauncher::FLaunchProfileTreeData& ProfileTreeData, FText OverrideDisplayName) const;
+		virtual UE_API FLaunchProfileTreeNode& AddDefaultHeading(ProjectLauncher::FLaunchProfileTreeData& ProfileTreeData) const;
 
 		/*
 		 * Populate the custom extension menu for this extension. This menu would typically contain
@@ -173,6 +164,11 @@ namespace ProjectLauncher
 		 */
 		virtual bool ManagesUATCommand( const ILauncherProfileUATCommandRef& UATCommand ) const { return false; }
 
+		/**
+		 * Get the associated UAT command, if any
+		 */
+		virtual ILauncherProfileUATCommandPtr GetUATCommand() const { return nullptr; }
+
 		/**
 		 * Internal intitialization function, called immediately after construction - don't call directly
 		 */
@@ -352,14 +348,11 @@ namespace ProjectLauncher
 		*/
 		UE_API void RequestFullTreeRebuild();
 
-	protected:
-		UE_API virtual FLaunchProfileTreeNode& AddDefaultHeadingInternal(FLaunchProfileTreeData& ProfileTreeData, FText DisplayName, int32 SortOrder) const;
-
 	private:
 
 		UE_API virtual void OnCustomValidation();
-		virtual int32 GetHeadingSortOrder() const { return 0; }
 
+		UE_API virtual const TCHAR* GetInternalName() const;
 
 		ILauncherProfileRef Profile;
 		ILauncherProfileRef DefaultProfile;
@@ -413,25 +406,25 @@ namespace ProjectLauncher
 		virtual bool IsCreatedByDefault( ILauncherProfileRef InProfile, TSharedRef<FModel> InModel ) const { return false; }
 
 		/**
-		 * Whether this extension can be created at this time
+		 * Whether this extension manages a UAT command
 		 */
-		virtual bool CanBeCreated( ILauncherProfileRef InProfile, TSharedRef<FModel> InModel ) const { return true; }
-
+		virtual bool IsUATCommandManager() const { return false; }
 
 		/**
-		 * Attempt to cast this launch extension to an automated test launch extension
+		 * Whether this extension can be created at this time
 		 */
-		virtual class FAutomatedTestLaunchExtension* AsAutomatedTestLaunchExtension() { return nullptr; }
+		virtual bool CanBeCreated( ILauncherProfileRef InProfile, TSharedRef<FModel> InModel ) const { return true; }
+
 
 		/**
-		 * Attempt to cast this launch extension to an UAT command launch extension
+		 * Attempt to cast this launch extension to an BuildCookRun launch extension (this is an extension for a build cook run UAT command)
 		 */
-		virtual class FUATCommandLaunchExtension* AsUATCommandLaunchExtension() { return nullptr; }
+		virtual class FBuildCookRunLaunchExtension* AsBuildCookRunLaunchExtension() { return nullptr; }
 
 		/**
-		 * Attempt to cast this launch extension to an BuildCookRun launch extension
+		 * Attempt to cast this launch extension instance to a UAT command factory
 		 */
-		virtual class FBuildCookRunCommandExtension* AsBuildCookRunCommandExtension() { return nullptr; }
+		virtual class IUATCommandExtensionFactory* AsUATCommandFactory() { return nullptr; }
 
 		/**
 		 * Instantiate all compatible extensions
@@ -443,9 +436,11 @@ namespace ProjectLauncher
 		 */
 		UE_API static TArray<TSharedPtr<FLaunchExtensionInstance>> CreateExtensionInstancesForProfile( ILauncherProfileRef InProfile, TSharedRef<FModel> InModel, FLaunchProfileTreeData& InOwnerTreeData );
 
-		UE_API static TSharedPtr<FLaunchExtensionInstance> CreateExtensionInstance( TSharedRef<FLaunchExtension> InExtension, ILauncherProfileRef InProfile, TSharedRef<FModel> InModel, FLaunchProfileTreeData& InOwnerTreeData );
-		UE_API static void DestroyExtensionInstance(TSharedRef<FLaunchExtension> InExtension, ILauncherProfileRef InProfile, TSharedRef<FModel> InModel);
+		UE_API static TSharedPtr<FLaunchExtensionInstance> CreateExtensionInstance( TSharedRef<FLaunchExtension> InExtension, ILauncherProfileRef InProfile, TSharedRef<FModel> InModel, FLaunchProfileTreeData& InOwnerTreeData, ILauncherProfileUATCommandPtr InUATCommand = nullptr );
+		UE_API static void DestroyExtensionInstance(TSharedRef<FLaunchExtensionInstance> InExtensionInstance, ILauncherProfileRef InProfile, TSharedRef<FModel> InModel);
+
 		UE_API static TSharedPtr<FLaunchExtensionInstance> GetExtensionInstance(TSharedRef<FLaunchExtension> InExtension, ILauncherProfileRef InProfile);
+		UE_API static TArray<TSharedRef<FLaunchExtensionInstance>> GetExtensionInstances(TSharedRef<FLaunchExtension> InExtension, ILauncherProfileRef InProfile);
 
 		UE_API static TArray<TSharedRef<FLaunchExtension>> GetExtensions();
 
diff --git a/Engine/Plugins/Developer/ProjectLauncher/Source/ProjectLauncher/Public/Extension/UATCommandLaunchExtension.h b/Engine/Plugins/Developer/ProjectLauncher/Source/ProjectLauncher/Public/Extension/UATCommandLaunchExtension.h
--- a/Engine/Plugins/Developer/ProjectLauncher/Source/ProjectLauncher/Public/Extension/UATCommandLaunchExtension.h
+++ b/Engine/Plugins/Developer/ProjectLauncher/Source/ProjectLauncher/Public/Extension/UATCommandLaunchExtension.h
@@ -9,78 +9,54 @@
 namespace ProjectLauncher
 {
 	/**
-	 * Helper base class for an UAT command launch extension
-	 * Creates an item in the extension menu that will toggle the UAT command on and off
-	 * note: It is not necessary for a UAT command launch extension to use this class
+	 * Helper base class for an extension that is tied to a UAT command in the profile
 	 */
-	class FUATCommandLaunchExtensionInstance : public FLaunchExtensionInstance
+	class FUATCommandLaunchExtensionInstanceBase : public FLaunchExtensionInstance
 	{
 	public:
-		UE_API FUATCommandLaunchExtensionInstance( FArgs& InArgs );
+		UE_API FUATCommandLaunchExtensionInstanceBase( FArgs& InArgs );
 
-		/**
-		 * Determine if this UAT Command is currently active
-		 */
-		UE_API bool IsUATCommandActive() const;
-
-
-		UE_API virtual bool ManagesUATCommand( const ILauncherProfileUATCommandRef& UATCommand ) const override;
-
-	protected:
-		UE_API virtual void OnAdded() override;
-		UE_API virtual void OnRemoved() override;
-
-
-		UE_API bool CustomizeUATCommandLine( const ILauncherProfileUATCommandRef& InUATCommand, FString& InOutCommandLine ) override;
+		UE_API virtual bool ManagesUATCommand( const ILauncherProfileUATCommandRef& InUATCommand ) const override;
+		UE_API virtual ILauncherProfileUATCommandPtr GetUATCommand() const override;
 		
 		/**
 		 * Advanced hook to allow any advanced modification of the command line when this extension's UAT Command is launched
 		 */
 		virtual void CustomizeUATCommandLine( FString& InOutCommandLine ) {}
 
-		/**
-		 * Hook to allow any modification of the UAT Command & profile once it has been created. At least the UAT Command name should be specified
-		 */
-		virtual void OnUATCommandAdded( ILauncherProfileUATCommandRef UATCommand ) = 0;
-
-		/**
-		 * Advanced hook to allow any modification of the profile just before the UAT Command is removed
-		 */
-		virtual void OnUATCommandRemoved( ILauncherProfileUATCommandRef UATCommand ) {}
-
-		/**
-		 * Access this extension's UAT Command if it has been created
-		 */
-		UE_API ILauncherProfileUATCommandPtr GetUATCommand() const;
+	protected:
+		const ILauncherProfileUATCommandRef UATCommand;
 
-		/**
-		 * Internal name for the UAT command. Must be globally unique, so recommended to include the extension name
-		 */
-		virtual const FString GetUATCommandInternalName() const = 0;
+		UE_API virtual bool CustomizeUATCommandLine( const ILauncherProfileUATCommandRef& InUATCommand, FString& InOutCommandLine ) override;
+		UE_API virtual void OnCustomValidation() override;
+		UE_API virtual FLaunchProfileTreeNode& AddDefaultHeading(FLaunchProfileTreeData& ProfileTreeData) const override;
+		UE_API virtual const TCHAR* GetInternalName() const override;
+	};
 
-		/**
-		 * Determine the order of this UAT command, negative runs before BuildCookRun, positive runs after
-		 */
-		UE_API virtual int32 GetUATCommandOrder() const { return INT32_MAX; }
+	/**
+	 * Interface for creating & maintaining UAT command extensions
+	 */
+	class IUATCommandExtensionFactory
+	{
+	public:
+		virtual TSharedPtr<ProjectLauncher::FLaunchExtensionInstance> CreateUATCommandExtensionInstance(ILauncherProfileRef InProfile, TSharedRef<FModel> InModel, FLaunchProfileTreeData& InOwnerTreeData) = 0;
 
-		/**
-		 * Internal intitialization function, called immediately after construction - don't call directly
-		 */
-		UE_API virtual void InternalInitialize() override;
+	protected:
+		virtual ILauncherProfileUATCommandRef CreateUATCommand(ILauncherProfileRef InProfile, TSharedRef<FModel> InModel) const = 0;
 
-	private:
-		UE_API virtual void OnCustomValidation() override;
-		virtual int32 GetHeadingSortOrder() const override { return GetUATCommandOrder(); }
-		UE_API virtual FLaunchProfileTreeNode& AddDefaultHeadingInternal(FLaunchProfileTreeData& ProfileTreeData, FText DisplayName, int32 SortOrder) const override;
 	};
 
 
-
-	class FUATCommandLaunchExtension : public FLaunchExtension
+	class FUATCommandLaunchExtensionBase : public FLaunchExtension, public IUATCommandExtensionFactory
 	{
 	public:
-		UE_API virtual void GetExtensionsMenuEntry( FExtensionsMenuEntry& MenuEntry ) const override;
-		virtual class FUATCommandLaunchExtension* AsUATCommandLaunchExtension() override { return this; }
+		virtual bool IsUATCommandManager() const { return true; }
+		
+		virtual IUATCommandExtensionFactory* AsUATCommandFactory() override { return this; }
+
+	protected:
+		UE_API virtual TSharedPtr<ProjectLauncher::FLaunchExtensionInstance> CreateUATCommandExtensionInstance(ILauncherProfileRef InProfile, TSharedRef<FModel> InModel, FLaunchProfileTreeData& InOwnerTreeData) override;
+
 	};
 
 };
diff --git a/Engine/Plugins/Developer/ProjectLauncher/Source/ProjectLauncher/Public/Model/ProjectLauncherModel.h b/Engine/Plugins/Developer/ProjectLauncher/Source/ProjectLauncher/Public/Model/ProjectLauncherModel.h
--- a/Engine/Plugins/Developer/ProjectLauncher/Source/ProjectLauncher/Public/Model/ProjectLauncherModel.h
+++ b/Engine/Plugins/Developer/ProjectLauncher/Source/ProjectLauncher/Public/Model/ProjectLauncherModel.h
@@ -229,9 +229,6 @@ namespace ProjectLauncher
 		/* Apply a content scheme to a profile, setting cook/deploy/packaging options etc */
 		UE_API void SetProfileContentScheme(EContentScheme ContentScheme, const ILauncherProfileBuildCookRunRef& BuildCookRun, bool bWantToCook = true, ELauncherProfileDeploymentModes::Type DefaultDeploymentMode = ELauncherProfileDeploymentModes::CopyToDevice);
 
-		/* Helper to add a BuildCookRun to the given profile with common initialized properties */
-		UE_API ILauncherProfileBuildCookRunRef CreateBuildCookRun( const ILauncherProfileRef& Profile );
-
 		/* Record a new log message from launching */
 		UE_API TSharedPtr<FLaunchLogMessage> AddLogMessage(const FString& InMessage, ELogVerbosity::Type InVerbosity = ELogVerbosity::Log);
 
@@ -304,6 +301,18 @@ namespace ProjectLauncher
 		UE_API TArray<FString> GetAvailableEngineMapNames();
 		UE_API TArray<FString> GetAvailableEngineMapPaths();
 
+		/* Get the current UAT command information for the given profile */
+		struct FUATCommandsDetail
+		{
+			int32 MinOrder = 0;
+			int32 MaxOrder = 0;
+			uint32 NumUATCommands = 0;
+		};
+		UE_API void GetUATCommandsDetail( const ILauncherProfileRef& Profile, FUATCommandsDetail& OutUATCommandsDetail);
+
+		/* Set the defaults for this BuildCookRun */
+		void SetDefaults(const ILauncherProfileRef& Profile, const ILauncherProfileBuildCookRunRef& BuildCookRun);
+
 	private:
 		/* Load and save config state */
 		UE_API void LoadConfig();
@@ -325,9 +334,6 @@ namespace ProjectLauncher
 		/* Get map paths for a project, caching results */
 		UE_API TArray<FString> GetAndCacheMapPaths(const FString& InOptionalProjectPath, bool bIncludeNonContentDirMaps = false);
 
-		/* Set the defaults for this BuildCookRun */
-		void SetDefaults(const ILauncherProfileRef& Profile, const ILauncherProfileBuildCookRunRef& BuildCookRun);
-
 	private:
 		FString ConfigFileName;
 		
diff --git a/Engine/Plugins/Performance/AutomatedPerfTesting/Source/AutomatedPerfTesting/Private/LaunchExtension/AutomatedPerfTestLaunchExtension.cpp b/Engine/Plugins/Performance/AutomatedPerfTesting/Source/AutomatedPerfTesting/Private/LaunchExtension/AutomatedPerfTestLaunchExtension.cpp
--- a/Engine/Plugins/Performance/AutomatedPerfTesting/Source/AutomatedPerfTesting/Private/LaunchExtension/AutomatedPerfTestLaunchExtension.cpp
+++ b/Engine/Plugins/Performance/AutomatedPerfTesting/Source/AutomatedPerfTesting/Private/LaunchExtension/AutomatedPerfTestLaunchExtension.cpp
@@ -49,7 +49,6 @@
 
 namespace AutomatedPerfTestConfig
 {
-	constexpr const TCHAR* AutomatedPerfTestInternalName = TEXT("AutomatedPerfTestExtension.PerfTest");
 	constexpr const TCHAR* EnableCSVProfilerConfig = TEXT("EnableCSVProfiler");
 	constexpr const TCHAR* WindowedConfig = TEXT("Windowed");
 	constexpr const TCHAR* IterationsConfig = TEXT("Iterations");
@@ -112,17 +111,6 @@ void FAutomatedPerfTestLaunchExtensionInstance::OnTestAdded(ILauncherProfileAuto
 	AutomatedTest->SetTests(TestNode);
 }
 
-int32 FAutomatedPerfTestLaunchExtensionInstance::GetTestOrder() const
-{
-    constexpr int32 DefaultOrder = 1000;
-	return DefaultOrder;
-}
-
-const FString FAutomatedPerfTestLaunchExtensionInstance::GetTestInternalName() const
-{
-	return AutomatedPerfTestInternalName;
-}
-
 void FAutomatedPerfTestLaunchExtensionInstance::CustomizeTree(ProjectLauncher::FLaunchProfileTreeData& ProfileTreeData)
 {
 	using namespace ProjectLauncher;
@@ -200,7 +188,7 @@ void FAutomatedPerfTestLaunchExtensionInstance::CustomizeTree(ProjectLauncher::F
 	AddTestNodeOptions(TreeNode);
 }
 
-void FAutomatedPerfTestLaunchExtensionInstance::CustomizeAutomatedTestCommandLine(FString& InOutCommandLine)
+void FAutomatedPerfTestLaunchExtensionInstance::CustomizeUATCommandLine(FString& InOutCommandLine)
 {
 	// Most of the parameters here are defined and handled in the corresponding 
 	// Gauntlet controller which launches the automated perf test with the required 
diff --git a/Engine/Plugins/Performance/AutomatedPerfTesting/Source/AutomatedPerfTesting/Private/LaunchExtension/AutomatedPerfTestLaunchExtension.h b/Engine/Plugins/Performance/AutomatedPerfTesting/Source/AutomatedPerfTesting/Private/LaunchExtension/AutomatedPerfTestLaunchExtension.h
--- a/Engine/Plugins/Performance/AutomatedPerfTesting/Source/AutomatedPerfTesting/Private/LaunchExtension/AutomatedPerfTestLaunchExtension.h
+++ b/Engine/Plugins/Performance/AutomatedPerfTesting/Source/AutomatedPerfTesting/Private/LaunchExtension/AutomatedPerfTestLaunchExtension.h
@@ -26,11 +26,9 @@ public:
 
 	virtual void OnPropertyChanged() override;
 	virtual void OnTestAdded(ILauncherProfileAutomatedTestRef AutomatedTest) override;
-	virtual const FString GetTestInternalName() const override;
-	virtual int32 GetTestOrder() const override;
 
 	virtual void CustomizeTree(ProjectLauncher::FLaunchProfileTreeData& ProfileTreeData) override;
-	virtual void CustomizeAutomatedTestCommandLine(FString& InOutCommandLine) override;
+	virtual void CustomizeUATCommandLine(FString& InOutCommandLine) override;
 
 private:
 	void ExportProfileGoScenarios(const FString& Filename);
@@ -51,8 +49,7 @@ private:
 	{
 		const auto IsTestType = [this]()
 		{
-			const bool bEnabled = IsTestActive();
-			return GetTestType() == TestType && bEnabled;
+			return GetTestType() == TestType;
 		};
 
 		ProjectLauncher::FLaunchProfileTreeNode::FCallbacks Callbacks
diff --git a/Engine/Source/Developer/LauncherServices/Private/Profiles/LauncherProfile.cpp b/Engine/Source/Developer/LauncherServices/Private/Profiles/LauncherProfile.cpp
--- a/Engine/Source/Developer/LauncherServices/Private/Profiles/LauncherProfile.cpp
+++ b/Engine/Source/Developer/LauncherServices/Private/Profiles/LauncherProfile.cpp
@@ -14,38 +14,39 @@ const TCHAR* FLauncherProfileAutomatedTest::TypeName = TEXT("AutomatedTest");
 const TCHAR* FLauncherProfileBuildCookRun::TypeName = TEXT("BuildCookRun");
 
 
-ILauncherProfileUATCommandRef FLauncherProfile::CreateUATCommandInternal(const TCHAR* InternalName)
+ILauncherProfileUATCommandRef FLauncherProfile::CreateUATCommandInternal(const TCHAR* InternalName, const TCHAR* UserTypeName)
 {
-	ILauncherProfileUATCommandRef UATCommand = MakeShared<FLauncherProfileUATCommand>(InternalName);
+	ILauncherProfileUATCommandRef UATCommand = MakeShared<FLauncherProfileUATCommand>(InternalName, UserTypeName);
 	return UATCommand;
 }
 
-ILauncherProfileAutomatedTestRef FLauncherProfile::CreateAutomatedTestInternal(const TCHAR* InternalName)
+ILauncherProfileAutomatedTestRef FLauncherProfile::CreateAutomatedTestInternal(const TCHAR* InternalName, const TCHAR* UserTypeName)
 {
-	ILauncherProfileAutomatedTestRef AutomatedTest = MakeShared<FLauncherProfileAutomatedTest>(InternalName);
+	ILauncherProfileAutomatedTestRef AutomatedTest = MakeShared<FLauncherProfileAutomatedTest>(InternalName, UserTypeName);
 	return AutomatedTest;
 }
 
-ILauncherProfileBuildCookRunRef FLauncherProfile::CreateBuildCookRunInternal(const TCHAR* InternalName)
+ILauncherProfileBuildCookRunRef FLauncherProfile::CreateBuildCookRunInternal(const TCHAR* InternalName, const TCHAR* UserTypeName)
 {
-	ILauncherProfileBuildCookRunRef BuildCookRun = MakeShared<FLauncherProfileBuildCookRun>(LauncherProfileManager, this, InternalName);
+	ILauncherProfileBuildCookRunRef BuildCookRun = MakeShared<FLauncherProfileBuildCookRun>(LauncherProfileManager, this, InternalName, UserTypeName);
 	return BuildCookRun;
 
 }
 
 ILauncherProfileUATCommandPtr FLauncherProfile::CreateUATCommandByTypeNameInternal( const TCHAR* TypeName, const TCHAR* InternalName )
 {
+	const TCHAR* UserTypeName = TEXT("");
 	if (FCString::Stricmp(TypeName, FLauncherProfileUATCommand::TypeName) == 0)
 	{
-		return CreateUATCommandInternal(InternalName);
+		return CreateUATCommandInternal(InternalName, UserTypeName);
 	}
 	else if (FCString::Stricmp(TypeName, FLauncherProfileAutomatedTest::TypeName) == 0)
 	{
-		return CreateAutomatedTestInternal(InternalName);
+		return CreateAutomatedTestInternal(InternalName, UserTypeName);
 	}
 	else if (FCString::Stricmp(TypeName, FLauncherProfileBuildCookRun::TypeName) == 0)
 	{
-		return CreateBuildCookRunInternal(InternalName);
+		return CreateBuildCookRunInternal(InternalName, UserTypeName);
 	}
 	else
 	{
@@ -468,7 +469,7 @@ bool FLauncherProfile::Serialize( FArchive& Archive )
 PRAGMA_DISABLE_DEPRECATION_WARNINGS
 			if (Archive.IsLoading())
 			{
-				ILauncherProfileAutomatedTestRef AutomatedTest = CreateAutomatedTestInternal();
+				ILauncherProfileAutomatedTestRef AutomatedTest = CreateAutomatedTestInternal(TEXT(""), TEXT(""));
 				AutomatedTest->Serialize(Archive);
 				AutomatedTests.Add(AutomatedTest);
 			}
diff --git a/Engine/Source/Developer/LauncherServices/Private/Profiles/LauncherProfile.h b/Engine/Source/Developer/LauncherServices/Private/Profiles/LauncherProfile.h
--- a/Engine/Source/Developer/LauncherServices/Private/Profiles/LauncherProfile.h
+++ b/Engine/Source/Developer/LauncherServices/Private/Profiles/LauncherProfile.h
@@ -1557,7 +1557,7 @@ public:
 				{
 					for (TSharedPtr<FJsonValue> Value : *Values)
 					{
-						ILauncherProfileUATCommandRef UATCommand = CreateAutomatedTestInternal();
+						ILauncherProfileUATCommandRef UATCommand = CreateAutomatedTestInternal(TEXT(""), TEXT(""));
 						if (UATCommand->Load((*Value->AsObject().Get())))
 						{
 							AddUATCommandInternal(UATCommand);
@@ -1586,7 +1586,7 @@ public:
 					const FJsonObject& JsonObject = *Value->AsObject().Get();
 					FString UATCommandType = JsonObject.GetStringField(TEXT("Type"));
 
-					ILauncherProfileUATCommandPtr UATCommand = CreateUATCommandByTypeNameInternal(*UATCommandType);
+					ILauncherProfileUATCommandPtr UATCommand = CreateUATCommandByTypeNameInternal(*UATCommandType, TEXT(""));
 					if (!UATCommand.IsValid())
 					{
 						UE_LOG(LogLauncherProfile, Error, TEXT("unknown UAT command type: %s"), *UATCommandType);
@@ -1674,7 +1674,7 @@ public:
 
 		UATCommands.Reset();
 		UATCommandValidation.Reset();
-		FirstBuildCookRun = CreateBuildCookRunInternal(TEXT("FirstBuildCookRun"));
+		FirstBuildCookRun = CreateBuildCookRunInternal(TEXT("FirstBuildCookRun"), TEXT(""));
 		AddUATCommandInternal(FirstBuildCookRun.ToSharedRef());
 		FirstBuildCookRun->SetDeployedDeviceGroup(LauncherProfileManager->AddNewDeviceGroup());
 
@@ -2261,12 +2261,20 @@ public:
 		return UATCommand ? UATCommand->AsAutomatedTest() : nullptr;
 	}
 
-	virtual ILauncherProfileAutomatedTestRef FindOrAddAutomatedTest( const TCHAR* InternalName ) override
+	virtual ILauncherProfileAutomatedTestRef FindOrAddAutomatedTest( const TCHAR* InternalName, const TCHAR* UserTypeName ) override
 	{
+		// if no name has been given, autogenerate one
+		FString UniqueName;
+		if (InternalName == nullptr)
+		{
+			UniqueName = GetUniqueUATCommandName( UserTypeName ? UserTypeName : TEXT("CustomAutomatedTest"));
+			InternalName = UniqueName.GetCharArray().GetData();
+		}
+
 		ILauncherProfileAutomatedTestPtr AutomatedTest = GetAutomatedTest(InternalName);
 		if (AutomatedTest == nullptr)
 		{
-			AutomatedTest = CreateAutomatedTestInternal(InternalName);
+			AutomatedTest = CreateAutomatedTestInternal(InternalName, UserTypeName);
 			AddUATCommandInternal(AutomatedTest.ToSharedRef());
 			UATCommandAdded.Broadcast(AutomatedTest.ToSharedRef());
 		}
@@ -2323,13 +2331,13 @@ public:
 		return UATCommand ? UATCommand->AsBuildCookRun() : nullptr;
 	}
 
-	virtual ILauncherProfileBuildCookRunRef FindOrAddBuildCookRunCommand( const TCHAR* InternalName ) override
+	virtual ILauncherProfileBuildCookRunRef FindOrAddBuildCookRunCommand( const TCHAR* InternalName, const TCHAR* UserTypeName ) override
 	{
 		// if no name has been given, autogenerate one
 		FString UniqueName;
 		if (InternalName == nullptr)
 		{
-			UniqueName = GetUniqueUATCommandName(TEXT("CustomBuildCookRun"));
+			UniqueName = GetUniqueUATCommandName(UserTypeName ? UserTypeName : TEXT("CustomBuildCookRun"));
 			InternalName = UniqueName.GetCharArray().GetData();
 		}
 
@@ -2337,7 +2345,7 @@ public:
 		ILauncherProfileBuildCookRunPtr BuildCookRun = GetBuildCookRunCommand(InternalName);
 		if (BuildCookRun == nullptr)
 		{
-			BuildCookRun = CreateBuildCookRunInternal(InternalName);
+			BuildCookRun = CreateBuildCookRunInternal(InternalName, UserTypeName);
 			BuildCookRun->SetDeployedDeviceGroup(LauncherProfileManager->AddNewDeviceGroup());
 
 			if (!FirstBuildCookRun.IsValid())
@@ -2370,12 +2378,20 @@ public:
 
 	}
 
-	virtual ILauncherProfileUATCommandRef FindOrAddUATCommand( const TCHAR* InternalName ) override
+	virtual ILauncherProfileUATCommandRef FindOrAddUATCommand( const TCHAR* InternalName, const TCHAR* UserTypeName ) override
 	{
+		// if no name has been given, autogenerate one
+		FString UniqueName;
+		if (InternalName == nullptr)
+		{
+			UniqueName = GetUniqueUATCommandName( UserTypeName ? UserTypeName : TEXT("CustomUATCommand"));
+			InternalName = UniqueName.GetCharArray().GetData();
+		}
+
 		ILauncherProfileUATCommandPtr UATCommand = GetUATCommand(InternalName);
 		if (UATCommand == nullptr)
 		{
-			UATCommand = CreateUATCommandInternal(InternalName);
+			UATCommand = CreateUATCommandInternal(InternalName, UserTypeName);
 
 			AddUATCommandInternal(UATCommand.ToSharedRef());
 			UATCommandAdded.Broadcast(UATCommand.ToSharedRef());
@@ -2952,10 +2968,10 @@ private:
 	}
 
 
-	ILauncherProfileUATCommandRef CreateUATCommandInternal( const TCHAR* InternalName = nullptr );
-	ILauncherProfileAutomatedTestRef CreateAutomatedTestInternal( const TCHAR* InternalName = nullptr );
-	ILauncherProfileBuildCookRunRef CreateBuildCookRunInternal( const TCHAR* InternalName = nullptr );
-	ILauncherProfileUATCommandPtr CreateUATCommandByTypeNameInternal( const TCHAR* TypeName, const TCHAR* InternalName = nullptr );
+	ILauncherProfileUATCommandRef CreateUATCommandInternal( const TCHAR* InternalName, const TCHAR* UserTypeName );
+	ILauncherProfileAutomatedTestRef CreateAutomatedTestInternal( const TCHAR* InternalName, const TCHAR* UserTypeName );
+	ILauncherProfileBuildCookRunRef CreateBuildCookRunInternal( const TCHAR* InternalName, const TCHAR* UserTypeName );
+	ILauncherProfileUATCommandPtr CreateUATCommandByTypeNameInternal( const TCHAR* TypeName, const TCHAR* InternalName );
 	ILauncherProfileBuildCookRunPtr LoadBuildCookRunInternal(const FJsonObject& Object, int32 Version);
 
 	void AddUATCommandInternal( const ILauncherProfileUATCommandRef& UATCommand )
diff --git a/Engine/Source/Developer/LauncherServices/Private/Profiles/LauncherProfileAutomatedTest.h b/Engine/Source/Developer/LauncherServices/Private/Profiles/LauncherProfileAutomatedTest.h
--- a/Engine/Source/Developer/LauncherServices/Private/Profiles/LauncherProfileAutomatedTest.h
+++ b/Engine/Source/Developer/LauncherServices/Private/Profiles/LauncherProfileAutomatedTest.h
@@ -13,8 +13,9 @@ class FLauncherProfileAutomatedTest final : public ILauncherProfileAutomatedTest
 public:
 	static const TCHAR* TypeName;
 
-	FLauncherProfileAutomatedTest( const TCHAR* InInternalName = TEXT("") ) 
+	FLauncherProfileAutomatedTest( const TCHAR* InInternalName = TEXT(""), const TCHAR* InUserTypeName = TEXT("") ) 
 		: InternalName(InInternalName)
+		, UserTypeName(InUserTypeName)
 	{
 	}
 
@@ -87,6 +88,16 @@ public:
 		return Description;
 	}
 
+	virtual void SetUserTypeName(const FString& NewUserTypeName) override
+	{
+		UserTypeName = NewUserTypeName;
+	}
+
+	virtual FString GetUserTypeName() const override
+	{
+		return UserTypeName;
+	}
+
 	virtual ILauncherProfileAutomatedTestPtr AsAutomatedTest() override
 	{
 		return AsShared().ToSharedPtr();
@@ -101,6 +112,7 @@ public:
 	{
 		Original = 0,
 		FirstVersioned = 1,
+		AddUserTypeName = 2,
 		// add new items here
 
 		Final,
@@ -126,18 +138,23 @@ public:
 		UATCommandLine        = Object.GetStringField(TEXT("AdditionalCommandLine"));
 		bEnabled              = Object.GetBoolField(TEXT("Enabled"));
 
-
 		if (Version >= EVersion::FirstVersioned)
 		{
 			Description = Object.GetStringField(TEXT("Description"));
 			Order       = Object.GetIntegerField(TEXT("Order"));
-
 		}
 		else
 		{
 			Object.TryGetNumberField(TEXT("Order"), Order);
 		}
 		Order = FMath::Max(1, Order);
+
+		UserTypeName = InternalName;
+		if (Version >= EVersion::AddUserTypeName)
+		{
+			UserTypeName = Object.GetStringField(TEXT("UserTypeName"));
+		}
+
 		return true;
 	}
 
@@ -153,6 +170,7 @@ public:
 		Writer.WriteValue(TEXT("Enabled"),               bEnabled);
 		Writer.WriteValue(TEXT("Order"),                 Order);
 		Writer.WriteValue(TEXT("Description"),           Description);
+		Writer.WriteValue(TEXT("UserTypeName"),          UserTypeName);
 		Writer.WriteObjectEnd();
 	}
 
@@ -179,6 +197,7 @@ protected:
 	int32 Order = 1; // longer tests should generally have a higher order
 
 	FString InternalName;
+	FString UserTypeName;
 	FString Description = NSLOCTEXT("FLauncherProfileAutomatedTest", "DefaultDescripton", "Automated Test").ToString();
 };
 
diff --git a/Engine/Source/Developer/LauncherServices/Private/Profiles/LauncherProfileBuildCookRun.h b/Engine/Source/Developer/LauncherServices/Private/Profiles/LauncherProfileBuildCookRun.h
--- a/Engine/Source/Developer/LauncherServices/Private/Profiles/LauncherProfileBuildCookRun.h
+++ b/Engine/Source/Developer/LauncherServices/Private/Profiles/LauncherProfileBuildCookRun.h
@@ -72,11 +72,12 @@ public:
 	/**
 	 * Default constructor.
 	 */
-	FLauncherProfileBuildCookRun( ILauncherProfileManagerRef InProfileManager, FLauncherProfile* InProfile, const TCHAR* InInternalName = TEXT("") ) 
+	FLauncherProfileBuildCookRun( ILauncherProfileManagerRef InProfileManager, FLauncherProfile* InProfile, const TCHAR* InInternalName = TEXT(""), const TCHAR* InUserTypeName = TEXT("") ) 
 		: LauncherProfileManager(InProfileManager)		
 		, DefaultLaunchRole(MakeShareable(new FLauncherProfileLaunchRole()))
 		, Profile(InProfile)
 		, InternalName(InInternalName)
+		, UserTypeName(InUserTypeName)
 	{ 
 		SetDefaults();
 		if (Profile)
@@ -679,6 +680,7 @@ public:
 		// ILauncherProfileUATCommand items
 		Writer.WriteValue("Type", TypeName);
 		Writer.WriteValue("InternalName", InternalName);
+		Writer.WriteValue("UserTypeName", UserTypeName);
 		Writer.WriteValue("UATCommand", UATCommand);
 		Writer.WriteValue("UATCommandLine", UATCommandLine);
 		Writer.WriteValue("Enabled", bEnabled);
@@ -855,6 +857,7 @@ public:
 	{
 		PreInnerBuildCookRun=0,
 		Original=1,
+		AddUserTypeName=2,
 		// add new items here
 
 		Final,
@@ -893,7 +896,11 @@ public:
 			Order                 = Object.GetIntegerField(TEXT("Order"));
 			Description           = Object.GetStringField(TEXT("Description"));
 		}
-
+		UserTypeName = InternalName;
+		if (Version >= EVersion::AddUserTypeName)
+		{
+			UserTypeName = Object.GetStringField(TEXT("UserTypeName"));
+		}
 
 		BuildConfiguration = (EBuildConfiguration)((int32)Object.GetNumberField(TEXT("BuildConfiguration")));
 		CookConfiguration = (EBuildConfiguration)((int32)Object.GetNumberField(TEXT("CookConfiguration")));
@@ -1822,6 +1829,16 @@ public:
 		return *InternalName;
 	}
 
+	virtual void SetUserTypeName(const FString& NewUserTypeName) override
+	{
+		UserTypeName = NewUserTypeName;
+	}
+
+	virtual FString GetUserTypeName() const override
+	{
+		return UserTypeName;
+	}
+
 	virtual ILauncherProfileAutomatedTestPtr AsAutomatedTest() override
 	{
 		return nullptr;
@@ -2075,6 +2092,7 @@ private:
 	int32 Order = 0;
 
 	FString InternalName;
+	FString UserTypeName;
 
 	// only for backwards compatibility in FLauncherProfile::Serialize
 	friend class FLauncherProfile;
diff --git a/Engine/Source/Developer/LauncherServices/Private/Profiles/LauncherProfileUATCommand.h b/Engine/Source/Developer/LauncherServices/Private/Profiles/LauncherProfileUATCommand.h
--- a/Engine/Source/Developer/LauncherServices/Private/Profiles/LauncherProfileUATCommand.h
+++ b/Engine/Source/Developer/LauncherServices/Private/Profiles/LauncherProfileUATCommand.h
@@ -14,8 +14,9 @@ class FLauncherProfileUATCommand final : public ILauncherProfileUATCommand, publ
 public:
 	static const TCHAR* TypeName;
 
-	FLauncherProfileUATCommand(const TCHAR* InInternalName = TEXT("") ) 
+	FLauncherProfileUATCommand(const TCHAR* InInternalName = TEXT(""), const TCHAR* InUserTypeName = TEXT("") ) 
 		: InternalName(InInternalName)
+		, UserTypeName(InUserTypeName)
 	{
 	}
 
@@ -66,11 +67,21 @@ public:
 		Description = NewDescription;
 	}
 
-	virtual FString GetDescription() const
+	virtual FString GetDescription() const override
 	{
 		return Description;
 	}
 
+	virtual void SetUserTypeName(const FString& NewUserTypeName) override
+	{
+		UserTypeName = NewUserTypeName;
+	}
+
+	virtual FString GetUserTypeName() const override
+	{
+		return UserTypeName;
+	}
+
 	virtual ILauncherProfileAutomatedTestPtr AsAutomatedTest() override
 	{
 		return nullptr;
@@ -85,6 +96,7 @@ public:
 	{
 		Original = 0,
 		FirstVersioned = 1,
+		AddUserTypeName = 2,
 		// add new items here
 
 		Final,
@@ -108,10 +120,17 @@ public:
 		UATCommandLine        = Object.GetStringField(TEXT("AdditionalCommandLine"));
 		bEnabled              = Object.GetBoolField(TEXT("Enabled"));
 		Order                 = Object.GetIntegerField(TEXT("Order"));
+
 		if (Version >= EVersion::FirstVersioned)
 		{
 			Description = Object.GetStringField(TEXT("Description"));
 		}
+		
+		UserTypeName = InternalName;	
+		if (Version >= EVersion::AddUserTypeName)
+		{
+			UserTypeName = Object.GetStringField(TEXT("UserTypeName"));
+		}
 		return true;
 	}
 
@@ -126,6 +145,7 @@ public:
 		Writer.WriteValue(TEXT("Enabled"),               bEnabled);
 		Writer.WriteValue(TEXT("Order"),                 Order);
 		Writer.WriteValue(TEXT("Description"),           Description);
+		Writer.WriteValue(TEXT("UserTypeName"),          UserTypeName);
 		Writer.WriteObjectEnd();
 	}
 
@@ -139,6 +159,7 @@ protected:
 	int32 Order = -100000; // negative order items will run before the main BuildCookRun, positive order items will run afterwards
 
 	FString InternalName;
+	FString UserTypeName;
 	FString Description = NSLOCTEXT("FLauncherProfileUATCommand", "DefaultDescripton", "Additional Task").ToString();
 };
 
diff --git a/Engine/Source/Developer/LauncherServices/Public/ILauncherProfile.h b/Engine/Source/Developer/LauncherServices/Public/ILauncherProfile.h
--- a/Engine/Source/Developer/LauncherServices/Public/ILauncherProfile.h
+++ b/Engine/Source/Developer/LauncherServices/Public/ILauncherProfile.h
@@ -1435,8 +1435,9 @@ public:
 
 	/**
 	 * Returns an UAT command by internal name, creating it if it doesn't already exist
+	 * Passing a null internal name will generate a new unique name
 	 */
-	virtual ILauncherProfileUATCommandRef FindOrAddUATCommand( const TCHAR* InternalName ) = 0;
+	virtual ILauncherProfileUATCommandRef FindOrAddUATCommand( const TCHAR* InternalName, const TCHAR* UserTypeName = nullptr ) = 0;
 
 	/**
 	 * Removes an UAT command by internal name
@@ -1472,8 +1473,9 @@ public:
 
 	/** 
 	 * Returns a UAT build cook run command by internal name, creating it if it doesn't already exist
+	 * Passing a null internal name will generate a new unique name
 	 */
-	virtual ILauncherProfileBuildCookRunRef FindOrAddBuildCookRunCommand( const TCHAR* InternalName ) = 0;
+	virtual ILauncherProfileBuildCookRunRef FindOrAddBuildCookRunCommand( const TCHAR* InternalName, const TCHAR* UserTypeName = nullptr ) = 0;
 
 
 
@@ -1490,8 +1492,9 @@ public:
 
 	/**
 	 * Returns an automated test by internal name, creating it if it doesn't already exist
+	 * Passing a null internal name will generate a new unique name
 	 */
-	virtual ILauncherProfileAutomatedTestRef FindOrAddAutomatedTest( const TCHAR* InternalName ) = 0;
+	virtual ILauncherProfileAutomatedTestRef FindOrAddAutomatedTest( const TCHAR* InternalName, const TCHAR* UserTypeName = nullptr ) = 0;
 
 	/**
 	 * Removes an automated test by internal name
diff --git a/Engine/Source/Developer/LauncherServices/Public/ILauncherProfileUATCommand.h b/Engine/Source/Developer/LauncherServices/Public/ILauncherProfileUATCommand.h
--- a/Engine/Source/Developer/LauncherServices/Public/ILauncherProfileUATCommand.h
+++ b/Engine/Source/Developer/LauncherServices/Public/ILauncherProfileUATCommand.h
@@ -35,6 +35,9 @@ public:
 	virtual void SetDescription(const FString& NewDescription) = 0;
 	virtual FString GetDescription() const = 0;
 
+	virtual void SetUserTypeName(const FString& NewUserTypeName) = 0;
+	virtual FString GetUserTypeName() const = 0;
+
 	virtual bool Load(const FJsonObject& Object) = 0;
 	virtual void Save(TJsonWriter<>& Writer) = 0;
 
