various property editor changes based on feedback:
- pin headings at the top as you scroll
- error message banner moved to the bottom
- use vertical ellipses for overflow menu

#rnx

#Synced-CL 50053775

diff --git a/Engine/Plugins/Developer/ProjectLauncher/Source/ProjectLauncher/Private/Styling/ProjectLauncherStyle.cpp b/Engine/Plugins/Developer/ProjectLauncher/Source/ProjectLauncher/Private/Styling/ProjectLauncherStyle.cpp
--- a/Engine/Plugins/Developer/ProjectLauncher/Source/ProjectLauncher/Private/Styling/ProjectLauncherStyle.cpp
+++ b/Engine/Plugins/Developer/ProjectLauncher/Source/ProjectLauncher/Private/Styling/ProjectLauncherStyle.cpp
@@ -80,6 +80,8 @@ void FProjectLauncherStyle::Initialize()
 	StyleSet->Set("Icons.WarningWithColor.Small", new CORE_IMAGE_BRUSH_SVG("../Editor/Slate/Starship/Common/AlertTriangleSolid", Icon16x16, FStyleColors::Warning));
 
 	StyleSet->Set( "PathPickerButton", new CORE_IMAGE_BRUSH( "Icons/ellipsis_12x", Icon12x12 ) );
+	StyleSet->Set( "OverflowButton", new CORE_IMAGE_BRUSH_SVG("Starship/Common/ellipsis-vertical-narrow", FVector2f(6.f, 24.f)));
+
 
 	StyleSet->Set("Icons.Task.Run",	new CORE_IMAGE_BRUSH_SVG("../Editor/Slate/Starship/Common/Timecode", Icon16x16));
 	StyleSet->Set("Icons.Task.Cleanup", new CORE_IMAGE_BRUSH_SVG("../Editor/Slate/Starship/GraphEditors/CleanUp", Icon16x16));
diff --git a/Engine/Plugins/Developer/ProjectLauncher/Source/ProjectLauncher/Private/Widgets/Profiles/SCustomLaunchCustomProfileEditor.cpp b/Engine/Plugins/Developer/ProjectLauncher/Source/ProjectLauncher/Private/Widgets/Profiles/SCustomLaunchCustomProfileEditor.cpp
--- a/Engine/Plugins/Developer/ProjectLauncher/Source/ProjectLauncher/Private/Widgets/Profiles/SCustomLaunchCustomProfileEditor.cpp
+++ b/Engine/Plugins/Developer/ProjectLauncher/Source/ProjectLauncher/Private/Widgets/Profiles/SCustomLaunchCustomProfileEditor.cpp
@@ -92,7 +92,7 @@ public:
 	bool bIsPrimaryHeading;
 	SCustomLaunchCustomProfileEditor* ProfileEditor;
 
-	void Construct(const FArguments& InArgs, ProjectLauncher::FLaunchProfileTreeNodePtr InTreeNode, SCustomLaunchCustomProfileEditor* InProfileEditor, bool bInIsPrimaryHeading )
+	void Construct(const FArguments& InArgs, ProjectLauncher::FLaunchProfileTreeNodePtr InTreeNode, SCustomLaunchCustomProfileEditor* InProfileEditor, bool bInIsPrimaryHeading, bool bPinned )
 	{
 		TreeNode = InTreeNode;
 		ProfileEditor = InProfileEditor;
@@ -192,8 +192,8 @@ public:
 			.Visibility_Lambda( [HasVisibleChildren]() { return HasVisibleChildren() ? EVisibility::Visible : EVisibility::Collapsed; } )
 			[
 				SNew(SBorder)
-				.BorderImage(FAppStyle::GetBrush("Brushes.Header"))				
-				.Padding(8*TreeNode->Depth,0,0,0)
+				.BorderImage(FAppStyle::GetBrush(bPinned ? "Brushes.Panel" : "Brushes.Header"))
+				.Padding(8*TreeNode->Depth, 0, bPinned ? 8 : 0, 0)
 				[
 					SNew(SBox)
 					.MinDesiredHeight(26.0f)
@@ -209,14 +209,14 @@ public:
 						[
 							SNew(SExpanderArrow, SharedThis(this))
 							.StyleSet( &FCoreStyle::Get() )
+							.Visibility( bPinned ? EVisibility::Hidden : EVisibility::Visible )
 						]
-
 					]
 				]
 			]
 		];
 
-		if (bIsPrimaryHeading)
+		if (bIsPrimaryHeading && !bPinned)
 		{
 			// checkbox to enable the UAT command
 			HeaderBar->AddSlot()
@@ -287,7 +287,7 @@ public:
 			];
 		}
 
-		if (bIsPrimaryHeading)
+		if (bIsPrimaryHeading && !bPinned)
 		{
 			// spacer to push everthing else to the right
 			HeaderBar->AddSlot()
@@ -309,8 +309,9 @@ public:
 				.OnGetMenuContent( ProfileEditor, &SCustomLaunchCustomProfileEditor::MakeUATCommandSubMenuWidget, TreeNode->UATCommand.ToSharedRef())
 				.ButtonContent()
 				[
-					SNew(STextBlock)
-					.Text(LOCTEXT("...", "..."))
+					SNew(SImage)
+					.Image(FProjectLauncherStyle::Get().GetBrush("OverflowButton"))
+					.ColorAndOpacity(FSlateColor::UseForeground())
 				]
 			];
 		}
@@ -655,6 +656,23 @@ void SCustomLaunchCustomProfileEditor::Construct(const FArguments& InArgs, const
 		.Padding(2)
 		[
 			SNew(SVerticalBox)
+
+			// main property tree
+			+SVerticalBox::Slot()
+			.FillHeight(1)
+			.Padding(2)
+			[
+				SAssignNew(TreeView, SLaunchProfileTreeView)
+				.TreeItemsSource(&TreeBuilder->GetProfileTree()->Nodes)
+				.SelectionMode( ESelectionMode::Single )
+				.OnGenerateRow( this, &SCustomLaunchCustomProfileEditor::OnGenerateWidgetForTreeNode, false )
+				.OnGeneratePinnedRow( this, &SCustomLaunchCustomProfileEditor::OnGenerateWidgetForTreeNode, true )
+				.OnGetChildren( this, &SCustomLaunchCustomProfileEditor::OnGetChildren)
+				.HandleDirectionalNavigation(false)
+				.ShouldStackHierarchyHeaders(true) // keep headings at the top when scrolling
+			]
+
+			// error banner
 			+SVerticalBox::Slot()
 			.AutoHeight()
 			.VAlign(VAlign_Center)
@@ -691,25 +709,13 @@ void SCustomLaunchCustomProfileEditor::Construct(const FArguments& InArgs, const
 					]
 				]
 			]
-
-			+SVerticalBox::Slot()
-			.FillHeight(1)
-			.Padding(2)
-			[
-				SAssignNew(TreeView, SLaunchProfileTreeView)
-				.TreeItemsSource(&TreeBuilder->GetProfileTree()->Nodes)
-				.SelectionMode( ESelectionMode::Single )
-				.OnGenerateRow( this, &SCustomLaunchCustomProfileEditor::OnGenerateWidgetForTreeNode )
-				.OnGetChildren( this, &SCustomLaunchCustomProfileEditor::OnGetChildren)
-				.HandleDirectionalNavigation(false)
-			]
 		]
 	];
 
 }
 END_SLATE_FUNCTION_BUILD_OPTIMIZATION
 
-TSharedRef<ITableRow> SCustomLaunchCustomProfileEditor::OnGenerateWidgetForTreeNode( ProjectLauncher::FLaunchProfileTreeNodePtr TreeNode, const TSharedRef<STableViewBase>& OwnerTable )
+TSharedRef<ITableRow> SCustomLaunchCustomProfileEditor::OnGenerateWidgetForTreeNode( ProjectLauncher::FLaunchProfileTreeNodePtr TreeNode, const TSharedRef<STableViewBase>& OwnerTable, bool bPinned )
 {
 	auto IsVisible = [TreeNode]()
 	{
@@ -740,7 +746,7 @@ TSharedRef<ITableRow> SCustomLaunchCustomProfileEditor::OnGenerateWidgetForTreeN
 	{
 		bool bIsPrimaryHeading = TreeNode->UATCommand.IsValid() && TreeBuilder->GetProfileTree()->UATCommandHeadings.Contains(TreeNode);
 
-		return SNew(SLaunchProfileCategoryTreeRow, TreeNode, this, bIsPrimaryHeading);
+		return SNew(SLaunchProfileCategoryTreeRow, TreeNode, this, bIsPrimaryHeading, bPinned);
 	}
 }
 
diff --git a/Engine/Plugins/Developer/ProjectLauncher/Source/ProjectLauncher/Private/Widgets/Profiles/SCustomLaunchCustomProfileEditor.h b/Engine/Plugins/Developer/ProjectLauncher/Source/ProjectLauncher/Private/Widgets/Profiles/SCustomLaunchCustomProfileEditor.h
--- a/Engine/Plugins/Developer/ProjectLauncher/Source/ProjectLauncher/Private/Widgets/Profiles/SCustomLaunchCustomProfileEditor.h
+++ b/Engine/Plugins/Developer/ProjectLauncher/Source/ProjectLauncher/Private/Widgets/Profiles/SCustomLaunchCustomProfileEditor.h
@@ -51,7 +51,7 @@ public:
 private:
 	virtual void Tick(const FGeometry& AllottedGeometry, const double InCurrentTime, const float InDeltaTime) override;
 
-	TSharedRef<ITableRow> OnGenerateWidgetForTreeNode( ProjectLauncher::FLaunchProfileTreeNodePtr TreeNode, const TSharedRef<STableViewBase>& OwnerTable );
+	TSharedRef<ITableRow> OnGenerateWidgetForTreeNode( ProjectLauncher::FLaunchProfileTreeNodePtr TreeNode, const TSharedRef<STableViewBase>& OwnerTable, bool bPinned );
 	void OnGetChildren(ProjectLauncher::FLaunchProfileTreeNodePtr TreeNode, TArray<ProjectLauncher::FLaunchProfileTreeNodePtr>& OutChildren);
 
 	void CacheCollapsedRecursive( TSet<FString>& Collapsed, const FString& NodeName, ProjectLauncher::FLaunchProfileTreeNodePtr Node) const;
