tasks menu fixes:
- put autotests and uat commands into their own section, rather than a submenu
- remove the 'add xxx' prefix for tasks because the menu is now just 'add' anyway
- default section is always kept at the top of tasks & extensions menu

#rnx

#Synced-CL 50053775

#Synced-CL 50053775

diff --git a/Engine/Plugins/Developer/ProjectLauncher/Source/ProjectLauncher/Private/Extension/LaunchExtension.cpp b/Engine/Plugins/Developer/ProjectLauncher/Source/ProjectLauncher/Private/Extension/LaunchExtension.cpp
--- a/Engine/Plugins/Developer/ProjectLauncher/Source/ProjectLauncher/Private/Extension/LaunchExtension.cpp
+++ b/Engine/Plugins/Developer/ProjectLauncher/Source/ProjectLauncher/Private/Extension/LaunchExtension.cpp
@@ -425,9 +425,16 @@ namespace ProjectLauncher
 	namespace
 	{
 		static const TCHAR* DefaultSectionName = TEXT("Common");
-		static const FText DefaultSectionDisplayName = LOCTEXT("CommonSectionName","General Extensions");
+		static const FText DefaultSectionDisplayName = LOCTEXT("CommonSectionName","Common");
+
+		static const TCHAR* AutoTestSectionName = TEXT("AutoTest");
+		static const FText AutoTestSectionDisplayName = LOCTEXT("AutoTestSectionName", "Automated Tests");
+
+		static const TCHAR* UATCmdsSectionName = TEXT("UATCmds");
+		static const FText UATCmdsSectionDisplayName = LOCTEXT("UATCommandsSubMenu", "Additional Tasks");
 	}
 
+
 	const FLaunchExtension::FExtensionsMenuEntry FLaunchExtension::FExtensionsMenuEntry::None =
 	{
 		.Type = FLaunchExtension::FExtensionsMenuEntry::Type_None,
@@ -438,6 +445,7 @@ namespace ProjectLauncher
 		.SectionName = DefaultSectionName,
 		.SectionDisplayName = DefaultSectionDisplayName,
 		.Type = FLaunchExtension::FExtensionsMenuEntry::Type_Section,
+		.bIsDefault = true,
 	};
 
 	const FLaunchExtension::FExtensionsMenuEntry FLaunchExtension::FExtensionsMenuEntry::OwnSection =
@@ -456,21 +464,17 @@ namespace ProjectLauncher
 
 	const FLaunchExtension::FExtensionsMenuEntry FLaunchExtension::FExtensionsMenuEntry::AutomatedTests =
 	{
-		.SectionName = DefaultSectionName,
-		.SubmenuName = TEXT("AutomatedTests"),
-		.SectionDisplayName = DefaultSectionDisplayName,
-		.SubmenuDisplayName = LOCTEXT("AutotestMenuSubMenu", "Automated Tests"),
-		.Type = FLaunchExtension::FExtensionsMenuEntry::Type_SubMenu,
+		.SectionName = AutoTestSectionName,
+		.SectionDisplayName = AutoTestSectionDisplayName,
+		.Type = FLaunchExtension::FExtensionsMenuEntry::Type_Section,
 	};
 
 
 	const FLaunchExtension::FExtensionsMenuEntry FLaunchExtension::FExtensionsMenuEntry::UATCommands =
 	{
-		.SectionName = DefaultSectionName,
-		.SubmenuName = TEXT("UATCommands"),
-		.SectionDisplayName = DefaultSectionDisplayName,
-		.SubmenuDisplayName = LOCTEXT("UATCommandsSubMenu", "More Build Steps"),
-		.Type = FLaunchExtension::FExtensionsMenuEntry::Type_SubMenu,
+		.SectionName = UATCmdsSectionName,
+		.SectionDisplayName = UATCmdsSectionDisplayName,
+		.Type = FLaunchExtension::FExtensionsMenuEntry::Type_Section,
 	};
 
 
diff --git a/Engine/Plugins/Developer/ProjectLauncher/Source/ProjectLauncher/Private/Widgets/Profiles/SCustomLaunchCustomProfileEditor.cpp b/Engine/Plugins/Developer/ProjectLauncher/Source/ProjectLauncher/Private/Widgets/Profiles/SCustomLaunchCustomProfileEditor.cpp
--- a/Engine/Plugins/Developer/ProjectLauncher/Source/ProjectLauncher/Private/Widgets/Profiles/SCustomLaunchCustomProfileEditor.cpp
+++ b/Engine/Plugins/Developer/ProjectLauncher/Source/ProjectLauncher/Private/Widgets/Profiles/SCustomLaunchCustomProfileEditor.cpp
@@ -959,6 +959,7 @@ void SCustomLaunchCustomProfileEditor::MakeMenu( FMenuBuilder& MenuBuilder, EMen
 		FText DisplayName;
 		TMap<FString,TSharedPtr<FExtensionMenuSubMenu>> SubMenus;
 		TArray<TSharedRef<ProjectLauncher::FLaunchExtension>> Extensions;
+		bool bIsDefault = false;
 	};
 
 	TMap<FString,FExtensionMenuSection> Sections;
@@ -1009,14 +1010,20 @@ void SCustomLaunchCustomProfileEditor::MakeMenu( FMenuBuilder& MenuBuilder, EMen
 		}
 		else
 		{
+			Section->bIsDefault |= ExtensionMenuEntry.bIsDefault;
 			Section->Extensions.Add(Extension);
 		}
 	}
 
 
-	// sort sections alphabetically
+	// sort sections alphabetically, keeping the default section at th top
 	Sections.ValueSort([](const FExtensionMenuSection& A, const FExtensionMenuSection& B)
 	{
+		if (A.bIsDefault != B.bIsDefault)
+		{
+			return A.bIsDefault;
+		}
+
 		return A.DisplayName.CompareToCaseIgnored(B.DisplayName) < 0; 
 	});
 
@@ -1073,7 +1080,7 @@ void SCustomLaunchCustomProfileEditor::MakeExtensionMenu( FMenuBuilder& MenuBuil
 	if (MenuType == EMenuType::UATCommands && Extension->IsUATCommandManager() && Extension->CanBeCreated(CurrentProfile.ToSharedRef(), Model.ToSharedRef()) )
 	{
 		MenuBuilder.AddMenuEntry(
-			FText::Format(LOCTEXT("NewUATCmdLabel", "Add New {0}"), Extension->GetDisplayName()),
+			Extension->GetDisplayName(),
 			FText::GetEmpty(),
 			FSlateIcon(Extension->GetIcon()),
 			FUIAction(
diff --git a/Engine/Plugins/Developer/ProjectLauncher/Source/ProjectLauncher/Public/Extension/LaunchExtension.h b/Engine/Plugins/Developer/ProjectLauncher/Source/ProjectLauncher/Public/Extension/LaunchExtension.h
--- a/Engine/Plugins/Developer/ProjectLauncher/Source/ProjectLauncher/Public/Extension/LaunchExtension.h
+++ b/Engine/Plugins/Developer/ProjectLauncher/Source/ProjectLauncher/Public/Extension/LaunchExtension.h
@@ -474,6 +474,8 @@ namespace ProjectLauncher
 				Type_SubMenu,		// appears in the extensions menu in a submenu of the SectionName section
 			};
 			EType Type = Type_None;
+
+			bool bIsDefault = false; // default sections should be kept at the top
 		};
 		UE_API virtual void GetExtensionsMenuEntry( FExtensionsMenuEntry& MenuEntry ) const;
 
