disable the pinned headings because it causes scrolling jitter + optimize .IsVisible so it is only checked during child node enumeration

#rnx

#Synced-CL 50053775

diff --git a/Engine/Plugins/Developer/ProjectLauncher/Source/ProjectLauncher/Private/Widgets/Profiles/SCustomLaunchCustomProfileEditor.cpp b/Engine/Plugins/Developer/ProjectLauncher/Source/ProjectLauncher/Private/Widgets/Profiles/SCustomLaunchCustomProfileEditor.cpp
--- a/Engine/Plugins/Developer/ProjectLauncher/Source/ProjectLauncher/Private/Widgets/Profiles/SCustomLaunchCustomProfileEditor.cpp
+++ b/Engine/Plugins/Developer/ProjectLauncher/Source/ProjectLauncher/Private/Widgets/Profiles/SCustomLaunchCustomProfileEditor.cpp
@@ -98,19 +98,6 @@ public:
 		ProfileEditor = InProfileEditor;
 		bIsPrimaryHeading = bInIsPrimaryHeading;
 
-		auto HasVisibleChildren = [this]()
-		{
-			for (ProjectLauncher::FLaunchProfileTreeNodePtr ChildNode : TreeNode->Children)
-			{
-				if (ChildNode->Callbacks.IsVisible == nullptr || ChildNode->Callbacks.IsVisible())
-				{
-					return true;
-				}
-			}
-
-			return false;
-		};
-
 		auto GetErrorVisibility = [this]()
 		{
 			if (!IsItemExpanded())
@@ -189,7 +176,6 @@ public:
 			SNew(SBorder)
 			.BorderImage(FAppStyle::GetBrush("Brushes.Background"))
 			.Padding(FMargin(0, 0, 0, 1))
-			.Visibility_Lambda( [HasVisibleChildren]() { return HasVisibleChildren() ? EVisibility::Visible : EVisibility::Collapsed; } )
 			[
 				SNew(SBorder)
 				.BorderImage(FAppStyle::GetBrush(bPinned ? "Brushes.Panel" : "Brushes.Header"))
@@ -669,7 +655,7 @@ void SCustomLaunchCustomProfileEditor::Construct(const FArguments& InArgs, const
 				.OnGeneratePinnedRow( this, &SCustomLaunchCustomProfileEditor::OnGenerateWidgetForTreeNode, true )
 				.OnGetChildren( this, &SCustomLaunchCustomProfileEditor::OnGetChildren)
 				.HandleDirectionalNavigation(false)
-				.ShouldStackHierarchyHeaders(true) // keep headings at the top when scrolling
+				//.ShouldStackHierarchyHeaders(true) // keep headings at the top when scrolling (@todo: disabled as it causes scrolling judder - bug in STreeView)
 			]
 
 			// error banner
@@ -717,20 +703,9 @@ END_SLATE_FUNCTION_BUILD_OPTIMIZATION
 
 TSharedRef<ITableRow> SCustomLaunchCustomProfileEditor::OnGenerateWidgetForTreeNode( ProjectLauncher::FLaunchProfileTreeNodePtr TreeNode, const TSharedRef<STableViewBase>& OwnerTable, bool bPinned )
 {
-	auto IsVisible = [TreeNode]()
-	{
-		if (TreeNode->Callbacks.IsVisible == nullptr)
-		{
-			return true;
-		}
-
-		return TreeNode->Callbacks.IsVisible();
-	};
-
 	if (TreeNode->Name.IsEmpty() && TreeNode->Widget.IsValid())
 	{
 		return SNew(SLaunchProfileFullWidthTreeRow, TreeNode, this)
-		.Visibility_Lambda( [IsVisible]() { return IsVisible() ? EVisibility::Visible : EVisibility::Collapsed; } )
 		.IsEnabled_Lambda( [TreeNode]() { return TreeNode->Callbacks.IsEnabled ? TreeNode->Callbacks.IsEnabled() : true; } )
 		;
 	}
@@ -739,13 +714,11 @@ TSharedRef<ITableRow> SCustomLaunchCustomProfileEditor::OnGenerateWidgetForTreeN
 		return SNew(SLaunchProfilePropertyTreeRow, TreeNode, this)
 		.SplitterValue_Lambda( [this]() { return SplitterPos; } )
 		.OnSplitterResized_Lambda( [this](float NewPos) { SplitterPos = FMath::Clamp(NewPos, 0.0f, 1.0f); } )
-		.Visibility_Lambda( [IsVisible]() { return IsVisible() ? EVisibility::Visible : EVisibility::Collapsed; } )
 		;
 	}
 	else
 	{
 		bool bIsPrimaryHeading = TreeNode->UATCommand.IsValid() && TreeBuilder->GetProfileTree()->UATCommandHeadings.Contains(TreeNode);
-
 		return SNew(SLaunchProfileCategoryTreeRow, TreeNode, this, bIsPrimaryHeading, bPinned);
 	}
 }
@@ -811,7 +784,15 @@ void SCustomLaunchCustomProfileEditor::BeginRenameUATCommand( const ILauncherPro
 
 void SCustomLaunchCustomProfileEditor::OnGetChildren(ProjectLauncher::FLaunchProfileTreeNodePtr TreeNode, TArray<ProjectLauncher::FLaunchProfileTreeNodePtr>& OutChildren)
 {
-	OutChildren = TreeNode->Children;
+	for ( const ProjectLauncher::FLaunchProfileTreeNodePtr& Child : TreeNode->Children)
+	{
+		if (Child->Callbacks.IsVisible.IsSet() && !Child->Callbacks.IsVisible() )
+		{
+			continue;
+		}
+
+		OutChildren.Add(Child);
+	}
 }
 
 FVector2D SCustomLaunchCustomProfileEditor::GetScrollDistance()
