ProjectLauncher - hook backends to selected build from buildsync extension

#Synced-CL 50096664

diff --git a/Engine/Plugins/Developer/ProjectLauncher/Source/ProjectLauncher/Private/Extension/LaunchExtension.cpp b/Engine/Plugins/Developer/ProjectLauncher/Source/ProjectLauncher/Private/Extension/LaunchExtension.cpp
--- a/Engine/Plugins/Developer/ProjectLauncher/Source/ProjectLauncher/Private/Extension/LaunchExtension.cpp
+++ b/Engine/Plugins/Developer/ProjectLauncher/Source/ProjectLauncher/Private/Extension/LaunchExtension.cpp
@@ -26,6 +26,19 @@ namespace ProjectLauncher
 	static FCriticalSection GExtensionsCS;
 	static std::atomic<bool> GExtensionsDestroyed = false;
 
+	TArray<TSharedPtr<FLaunchExtensionInstance>> GetProfileExtensionInstances(ILauncherProfilePtr Profile)
+	{
+		// take a snapshot of the extension instances for this profile
+		TArray<TSharedPtr<FLaunchExtensionInstance>> ExtensionInstances;
+		{
+			FScopeLock Lock(&GExtensionsCS);
+			if (GExtensionInstances.Contains(Profile))
+			{
+				ExtensionInstances = GExtensionInstances[Profile];
+			}
+		}
+		return ExtensionInstances;
+	}
 
 	void ShutdownLaunchExtensions()
 	{
@@ -71,16 +84,7 @@ namespace ProjectLauncher
 
 	void ApplyExtensionVariables( const ILauncherProfileRef& InProfile, FString& InOutCommandLine, TSharedRef<FModel> InModel, EBuildTargetType InBuildTargetType )
 	{
-		// take a snapshot of the extension instances for this profile
-		TArray<TSharedPtr<FLaunchExtensionInstance>> ExtensionInstances;
-		{
-			FScopeLock Lock(&GExtensionsCS);
-			if (GExtensionInstances.Contains(InProfile))
-			{
-				ExtensionInstances = GExtensionInstances[InProfile];
-			}
-		}
-
+		TArray<TSharedPtr<FLaunchExtensionInstance>> ExtensionInstances = GetProfileExtensionInstances(InProfile);
 
 		for (TSharedPtr<FLaunchExtensionInstance> ExtensionInstance : ExtensionInstances)
 		{
@@ -119,16 +123,7 @@ namespace ProjectLauncher
 
 	void ApplyExtensionVariablesForUATCommand( FUATCommandPostProcessParameters& InPostProcessParameters, TSharedRef<FModel> InModel )
 	{
-		// take a snapshot of the extension instances for this profile
-		TArray<TSharedPtr<FLaunchExtensionInstance>> ExtensionInstances;
-		{
-			FScopeLock Lock(&GExtensionsCS);
-			if (GExtensionInstances.Contains(InPostProcessParameters.Profile))
-			{
-				ExtensionInstances = GExtensionInstances[InPostProcessParameters.Profile];
-			}
-		}
-
+		TArray<TSharedPtr<FLaunchExtensionInstance>> ExtensionInstances = GetProfileExtensionInstances(InPostProcessParameters.Profile);
 
 		for (TSharedPtr<FLaunchExtensionInstance> ExtensionInstance : ExtensionInstances)
 		{
@@ -367,7 +362,6 @@ namespace ProjectLauncher
 		}
 	}
 
-
 	TArray<TSharedRef<FLaunchExtensionInstance>> FLaunchExtension::GetExtensionInstances(TSharedRef<FLaunchExtension> InExtension, ILauncherProfileRef InProfile)
 	{
 		check(!GExtensionsDestroyed);
@@ -845,6 +839,37 @@ namespace ProjectLauncher
 	{
 		return Extension->GetInternalName();
 	}
+
+	void FLaunchExtensionInstance::BroadcastEvent(const FString& EventName, void* EventData)
+	{
+		TArray<TSharedPtr<FLaunchExtensionInstance>> ExtensionInstances = GetProfileExtensionInstances(Profile);
+		for (TSharedPtr<FLaunchExtensionInstance> ExtensionInstance : ExtensionInstances)
+		{
+			// Broadcast to all other extension instances associated with the same profile
+			if (ExtensionInstance.IsValid() && ExtensionInstance.Get() != this && ExtensionInstance->GetProfile() == GetProfile())
+			{
+				ExtensionInstance->HandleEventCallback(FString::Printf(TEXT("%s_%s"), GetInternalName(), *EventName), EventData);
+			}		
+		}
+	}
+
+	TArray<TSharedPtr<FLaunchExtensionInstance>> FLaunchExtensionInstance::GetProfileExtensionInstancesByName(const FString& InternalName, ILauncherProfileRef InProfile)
+	{
+		TArray<TSharedPtr<FLaunchExtensionInstance>> ExtensionInstances = GetProfileExtensionInstances(InProfile);
+		TArray<TSharedPtr<FLaunchExtensionInstance>> MatchedExtensionInstances;
+		for (TSharedPtr<FLaunchExtensionInstance> ExtensionInstance : ExtensionInstances)
+		{
+			if (ExtensionInstance.IsValid())
+			{
+				TSharedRef<FLaunchExtension> Extension = ExtensionInstance->GetExtension();
+				if (FString(Extension->GetInternalName()) == InternalName)
+				{
+					MatchedExtensionInstances.Add(ExtensionInstance);
+				}
+			}
+		}
+		return MatchedExtensionInstances;
+	}
 }
 
 #undef LOCTEXT_NAMESPACE
diff --git a/Engine/Plugins/Developer/ProjectLauncher/Source/ProjectLauncher/Public/Extension/LaunchExtension.h b/Engine/Plugins/Developer/ProjectLauncher/Source/ProjectLauncher/Public/Extension/LaunchExtension.h
--- a/Engine/Plugins/Developer/ProjectLauncher/Source/ProjectLauncher/Public/Extension/LaunchExtension.h
+++ b/Engine/Plugins/Developer/ProjectLauncher/Source/ProjectLauncher/Public/Extension/LaunchExtension.h
@@ -183,6 +183,21 @@ namespace ProjectLauncher
 		UE_API virtual void OnAdded() {};
 		UE_API virtual void OnRemoved() {};
 
+		/**
+		 * Broadcast an event to all other extension instances associated to the same profile.
+		 */
+		UE_API virtual void BroadcastEvent(const FString& EventName, void* EventData = nullptr);
+
+		/**
+		 * The listener checks the incoming event name and chooses to handle the event accordingly.
+		 */
+		UE_API virtual void HandleEventCallback(const FString& EventName, void* EventData) {}
+
+		/**
+		 * Search for the extension instances associated with the given profile by name
+		 */
+		UE_API static TArray<TSharedPtr<FLaunchExtensionInstance>> GetProfileExtensionInstancesByName(const FString& ExtensionName, ILauncherProfileRef InProfile);
+
 	protected:
 
 		UE_DEPRECATED(5.8, "Use ProjectLauncher::CmdLineUtils version")
