property tree improvements:
- added a spacer between UAT command headings
- easier access to General Settings header, allowing launch extensions to add properties to that section directly
- launch extensions add themselves as a child of General Settings by default
- in the basic launch profile, build cook run extensions are a child of General Settings
- fix for empty headings being visible
- fix for collapsed heading reopening when a new task is added
- new uat nodes are scrolled into view and the user is offered the chance to rename them automatically

#rnx

#Synced-CL 50098631

diff --git a/Engine/Plugins/Developer/ProjectLauncher/Source/ProjectLauncher/Private/Extension/BuildCookRunCommandExtension.cpp b/Engine/Plugins/Developer/ProjectLauncher/Source/ProjectLauncher/Private/Extension/BuildCookRunCommandExtension.cpp
--- a/Engine/Plugins/Developer/ProjectLauncher/Source/ProjectLauncher/Private/Extension/BuildCookRunCommandExtension.cpp
+++ b/Engine/Plugins/Developer/ProjectLauncher/Source/ProjectLauncher/Private/Extension/BuildCookRunCommandExtension.cpp
@@ -38,17 +38,7 @@ namespace ProjectLauncher
 
 	FLaunchProfileTreeNode& FBuildCookRunExtension::AddDefaultHeading(FLaunchProfileTreeNode& ProfileTreeData) const
 	{
-		return AddDefaultHeading(ProfileTreeData, Extension->GetDisplayName() );
-	}
-
-	FLaunchProfileTreeNode& FBuildCookRunExtension::AddDefaultHeading(FLaunchProfileTreeNode& ProfileTreeData, FText OverrideDisplayName) const
-	{
-		if (Model->CanUseSimplifiedLayout(Profile))
-		{
-			return ProfileTreeData.GetTreeData()->AddHeading(Extension->GetInternalName(), OverrideDisplayName, BuildCookRun->GetOrder() );
-		}
-
-		return ProfileTreeData.AddSubHeading(Extension->GetInternalName(), OverrideDisplayName, BuildCookRun->GetOrder() );
+		return ProfileTreeData.AddSubHeading(Extension->GetInternalName(), Extension->GetDisplayName(), INT32_MAX );
 	}
 
 
@@ -90,10 +80,17 @@ namespace ProjectLauncher
 
 		for (const TTuple<ILauncherProfileUATCommandPtr,TSharedRef<FBuildCookRunExtension>>& Pair : BuildCookRunExtensions)
 		{
-			const bool bIsPrimaryHeading = true;
-			FLaunchProfileTreeNode& ProfileTreeNode = ProfileTreeData.FindOrAddHeading(Pair.Key.ToSharedRef(), bIsPrimaryHeading);
+			if (GetModel()->CanUseSimplifiedLayout(GetProfile()))
+			{
+				Pair.Value->CustomizeTree(ProfileTreeData.GetHeading(FLaunchProfileTreeData::GeneralSettingsSectionName));
+			}
+			else
+			{
+				const bool bIsPrimaryHeading = true;
+				FLaunchProfileTreeNode& ProfileTreeNode = ProfileTreeData.FindOrAddHeading(Pair.Key.ToSharedRef(), bIsPrimaryHeading);
 
-			Pair.Value->CustomizeTree(ProfileTreeNode);
+				Pair.Value->CustomizeTree(ProfileTreeNode);
+			}
 		}
 	}
 
diff --git a/Engine/Plugins/Developer/ProjectLauncher/Source/ProjectLauncher/Private/Extension/LaunchExtension.cpp b/Engine/Plugins/Developer/ProjectLauncher/Source/ProjectLauncher/Private/Extension/LaunchExtension.cpp
--- a/Engine/Plugins/Developer/ProjectLauncher/Source/ProjectLauncher/Private/Extension/LaunchExtension.cpp
+++ b/Engine/Plugins/Developer/ProjectLauncher/Source/ProjectLauncher/Private/Extension/LaunchExtension.cpp
@@ -824,7 +824,8 @@ namespace ProjectLauncher
 
 	FLaunchProfileTreeNode& FLaunchExtensionInstance::AddDefaultHeading(ProjectLauncher::FLaunchProfileTreeData& ProfileTreeData) const
 	{
-		return ProfileTreeData.AddHeading(GetInternalName(), Extension->GetDisplayName() );
+		return ProfileTreeData.GetHeading(FLaunchProfileTreeData::GeneralSettingsSectionName)
+			.AddSubHeading(GetInternalName(), Extension->GetDisplayName(), INT32_MAX ); // new subheadings go at the end, after all of the properties in general settings
 	}
 
 	bool FLaunchExtensionInstance::CustomizeUATCommandLine( const ILauncherProfileUATCommandRef& InUATCommand, FString& InOutCommandLine )
diff --git a/Engine/Plugins/Developer/ProjectLauncher/Source/ProjectLauncher/Private/ProfileTree/BasicProfileTreeBuilder.cpp b/Engine/Plugins/Developer/ProjectLauncher/Source/ProjectLauncher/Private/ProfileTree/BasicProfileTreeBuilder.cpp
--- a/Engine/Plugins/Developer/ProjectLauncher/Source/ProjectLauncher/Private/ProfileTree/BasicProfileTreeBuilder.cpp
+++ b/Engine/Plugins/Developer/ProjectLauncher/Source/ProjectLauncher/Private/ProfileTree/BasicProfileTreeBuilder.cpp
@@ -20,7 +20,7 @@ namespace ProjectLauncher
 		ILauncherProfileBuildCookRunRef UATCommand = Profile->GetFirstBuildCookRun().ToSharedRef();
 		FBuildCookRun& BuildCookRun = Get(UATCommand);
 
-		FLaunchProfileTreeNode& GeneralSettingsHeader = TreeData->AddHeading( TEXT("GeneralSettings"), LOCTEXT("GeneralSettingsHeading", "General Settings"), INT32_MIN );
+		FLaunchProfileTreeNode& GeneralSettingsHeader = TreeData->AddHeading( FLaunchProfileTreeData::GeneralSettingsSectionName, LOCTEXT("GeneralSettingsHeading", "General Settings"), FLaunchProfileTreeData::GeneralSettingsSortOrder );
 		GeneralSettingsHeader.UATCommand = UATCommand;
 
 		AddProjectProperty(GeneralSettingsHeader);
diff --git a/Engine/Plugins/Developer/ProjectLauncher/Source/ProjectLauncher/Private/ProfileTree/CustomProfileTreeBuilder.cpp b/Engine/Plugins/Developer/ProjectLauncher/Source/ProjectLauncher/Private/ProfileTree/CustomProfileTreeBuilder.cpp
--- a/Engine/Plugins/Developer/ProjectLauncher/Source/ProjectLauncher/Private/ProfileTree/CustomProfileTreeBuilder.cpp
+++ b/Engine/Plugins/Developer/ProjectLauncher/Source/ProjectLauncher/Private/ProfileTree/CustomProfileTreeBuilder.cpp
@@ -36,7 +36,7 @@ namespace ProjectLauncher
 			}
 		};
 
-		FLaunchProfileTreeNode& GeneralSettingsHeader = TreeData->AddHeading( TEXT("GeneralSettings"), LOCTEXT("GeneralSettingsHeading", "General Settings"), INT32_MIN );
+		FLaunchProfileTreeNode& GeneralSettingsHeader = TreeData->AddHeading( FLaunchProfileTreeData::GeneralSettingsSectionName, LOCTEXT("GeneralSettingsHeading", "General Settings"), FLaunchProfileTreeData::GeneralSettingsSortOrder );
 		if (bClassicView)
 		{
 			GeneralSettingsHeader.UATCommand = Profile->GetFirstBuildCookRun();
diff --git a/Engine/Plugins/Developer/ProjectLauncher/Source/ProjectLauncher/Private/ProfileTree/LaunchProfileTreeData.cpp b/Engine/Plugins/Developer/ProjectLauncher/Source/ProjectLauncher/Private/ProfileTree/LaunchProfileTreeData.cpp
--- a/Engine/Plugins/Developer/ProjectLauncher/Source/ProjectLauncher/Private/ProfileTree/LaunchProfileTreeData.cpp
+++ b/Engine/Plugins/Developer/ProjectLauncher/Source/ProjectLauncher/Private/ProfileTree/LaunchProfileTreeData.cpp
@@ -111,7 +111,9 @@ namespace ProjectLauncher
 		}
 	}
 
-
+	const TCHAR* FLaunchProfileTreeData::GeneralSettingsSectionName = TEXT("GeneralSettings");
+	const int32 FLaunchProfileTreeData::GeneralSettingsSortOrder = INT32_MIN;
+	const int32 FLaunchProfileTreeData::DefaultSortOrder = INT32_MIN+1;
 
 	FLaunchProfileTreeData::FLaunchProfileTreeData(ILauncherProfilePtr InProfile, TSharedRef<ProjectLauncher::FModel> InModel, ILaunchProfileTreeBuilder* InTreeBuilder) 
 		: Profile(InProfile)
@@ -308,6 +310,12 @@ namespace ProjectLauncher
 		TreeNode->UATCommand = UATCommand;
 		Children.Add(TreeNode);
 
+		Children.StableSort( []( const FLaunchProfileTreeNodePtr& A, const FLaunchProfileTreeNodePtr& B)
+		{
+			return A->SortOrder < B->SortOrder;
+		});
+
+
 		TreeData->RequestTreeRefresh();
 		return *this;
 	}
diff --git a/Engine/Plugins/Developer/ProjectLauncher/Source/ProjectLauncher/Private/Widgets/Profiles/SCustomLaunchCustomProfileEditor.cpp b/Engine/Plugins/Developer/ProjectLauncher/Source/ProjectLauncher/Private/Widgets/Profiles/SCustomLaunchCustomProfileEditor.cpp
--- a/Engine/Plugins/Developer/ProjectLauncher/Source/ProjectLauncher/Private/Widgets/Profiles/SCustomLaunchCustomProfileEditor.cpp
+++ b/Engine/Plugins/Developer/ProjectLauncher/Source/ProjectLauncher/Private/Widgets/Profiles/SCustomLaunchCustomProfileEditor.cpp
@@ -175,7 +175,7 @@ public:
 		[
 			SNew(SBorder)
 			.BorderImage(FAppStyle::GetBrush("Brushes.Background"))
-			.Padding(FMargin(0, 0, 0, 1))
+			.Padding(FMargin(0, bIsPrimaryHeading ? 8 : 0, 0, 1))
 			[
 				SNew(SBorder)
 				.BorderImage(FAppStyle::GetBrush(bPinned ? "Brushes.Panel" : "Brushes.Header"))
@@ -782,16 +782,47 @@ void SCustomLaunchCustomProfileEditor::BeginRenameUATCommand( const ILauncherPro
 	}
 }
 
-void SCustomLaunchCustomProfileEditor::OnGetChildren(ProjectLauncher::FLaunchProfileTreeNodePtr TreeNode, TArray<ProjectLauncher::FLaunchProfileTreeNodePtr>& OutChildren)
+bool SCustomLaunchCustomProfileEditor::IsHierarchyVisible(ProjectLauncher::FLaunchProfileTreeNodePtr TreeNode) const
 {
+	// don't show empty headings
+	if (!TreeNode->Name.IsEmpty() && !TreeNode->Widget.IsValid() && TreeNode->Children.Num() == 0)
+	{
+		return false;
+	}
+
+	// don't show if we are invisible
+	if (TreeNode->Callbacks.IsVisible.IsSet() && !TreeNode->Callbacks.IsVisible() )
+	{
+		return false;
+	}
+
+	// no children - show by default
+	if (TreeNode->Children.Num() == 0)
+	{
+		return true;
+	}
+
+	// only show if we have a visible child
 	for ( const ProjectLauncher::FLaunchProfileTreeNodePtr& Child : TreeNode->Children)
 	{
-		if (Child->Callbacks.IsVisible.IsSet() && !Child->Callbacks.IsVisible() )
+		if (IsHierarchyVisible(Child))
 		{
-			continue;
+			return true;
 		}
+	}
+
+	return false;
+}
+
 
-		OutChildren.Add(Child);
+void SCustomLaunchCustomProfileEditor::OnGetChildren(ProjectLauncher::FLaunchProfileTreeNodePtr TreeNode, TArray<ProjectLauncher::FLaunchProfileTreeNodePtr>& OutChildren)
+{
+	for ( const ProjectLauncher::FLaunchProfileTreeNodePtr& Child : TreeNode->Children)
+	{
+		if (IsHierarchyVisible(Child))
+		{
+			OutChildren.Add(Child);
+		}
 	}
 }
 
@@ -874,14 +905,10 @@ void SCustomLaunchCustomProfileEditor::RebuildTree()
 
 	TreeBuilder = ProjectLauncher::CreateTreeBuilder( CurrentProfile, Model.ToSharedRef() );
 	TreeView->SetTreeItemsSource( &TreeBuilder->GetProfileTree()->Nodes );
-
-	// restore the collapsed state
-	for ( TTuple<FString,ProjectLauncher::FLaunchProfileTreeNodePtr> Pair : TreeBuilder->GetProfileTree()->HeadingNodes)
-	{
-		ApplyCollapsedRecusive(CurrentlyCollapsed, Pair.Key, Pair.Value);
-	}
-
 	TreeView->RequestTreeRefresh();
+
+	// restore the collapsed state next tick, once the tree has refreshed
+	PendingCollapseHeadings = MoveTemp(CurrentlyCollapsed);
 }
 
 void SCustomLaunchCustomProfileEditor::CacheCollapsedRecursive( TSet<FString>& Collapsed, const FString& NodeName, ProjectLauncher::FLaunchProfileTreeNodePtr Node) const
@@ -899,7 +926,7 @@ void SCustomLaunchCustomProfileEditor::CacheCollapsedRecursive( TSet<FString>& C
 	}
 }
 
-void SCustomLaunchCustomProfileEditor::ApplyCollapsedRecusive( TSet<FString>& Collapsed, const FString& NodeName, ProjectLauncher::FLaunchProfileTreeNodePtr Node) const
+void SCustomLaunchCustomProfileEditor::ApplyCollapsedRecusive( const TSet<FString>& Collapsed, const FString& NodeName, ProjectLauncher::FLaunchProfileTreeNodePtr Node) const
 {
 	if (!Collapsed.Contains(NodeName))
 	{
@@ -1071,6 +1098,7 @@ void SCustomLaunchCustomProfileEditor::MakeExtensionMenu( FMenuBuilder& MenuBuil
 					if (Instance.IsValid())
 					{
 						TreeBuilder->GetProfileTree()->RequestFullTreeRebuild();
+						PendingNewHeadingAdded = Instance->GetUATCommand()->GetInternalName();
 					}
 				})
 			),
@@ -1361,6 +1389,37 @@ void SCustomLaunchCustomProfileEditor::Tick(const FGeometry& AllottedGeometry, c
 			TreeView->RequestTreeRefresh();
 			TreeBuilder->GetProfileTree()->bRequestTreeRefresh = false;
 		}
+		else if (!PendingNewHeadingAdded.IsEmpty())
+		{
+			// find the heading node
+			for ( TTuple<FString,ProjectLauncher::FLaunchProfileTreeNodePtr> Pair : TreeBuilder->GetProfileTree()->HeadingNodes)
+			{
+				if (Pair.Key == PendingNewHeadingAdded)
+				{
+					TreeView->RequestScrollIntoView(Pair.Value);
+
+					// start renaming it until after the new node has been scrolled into view
+					ExecuteOnGameThread(UE_SOURCE_LOCATION, [this, UATCommand = Pair.Value->UATCommand.ToSharedRef()]()
+					{
+						BeginRenameUATCommand(UATCommand);
+					});
+
+					break;
+				}
+			}
+			PendingNewHeadingAdded.Reset();
+		}
+		
+		if (PendingCollapseHeadings.Num() > 0)
+		{
+			for ( TTuple<FString,ProjectLauncher::FLaunchProfileTreeNodePtr> Pair : TreeBuilder->GetProfileTree()->HeadingNodes)
+			{
+				ApplyCollapsedRecusive(PendingCollapseHeadings, Pair.Key, Pair.Value);
+			}
+
+			PendingCollapseHeadings.Reset();
+		}
+
 	}
 }
 
diff --git a/Engine/Plugins/Developer/ProjectLauncher/Source/ProjectLauncher/Private/Widgets/Profiles/SCustomLaunchCustomProfileEditor.h b/Engine/Plugins/Developer/ProjectLauncher/Source/ProjectLauncher/Private/Widgets/Profiles/SCustomLaunchCustomProfileEditor.h
--- a/Engine/Plugins/Developer/ProjectLauncher/Source/ProjectLauncher/Private/Widgets/Profiles/SCustomLaunchCustomProfileEditor.h
+++ b/Engine/Plugins/Developer/ProjectLauncher/Source/ProjectLauncher/Private/Widgets/Profiles/SCustomLaunchCustomProfileEditor.h
@@ -55,9 +55,11 @@ private:
 	void OnGetChildren(ProjectLauncher::FLaunchProfileTreeNodePtr TreeNode, TArray<ProjectLauncher::FLaunchProfileTreeNodePtr>& OutChildren);
 
 	void CacheCollapsedRecursive( TSet<FString>& Collapsed, const FString& NodeName, ProjectLauncher::FLaunchProfileTreeNodePtr Node) const;
-	void ApplyCollapsedRecusive( TSet<FString>& Collapsed, const FString& NodeName, ProjectLauncher::FLaunchProfileTreeNodePtr Node) const;
+	void ApplyCollapsedRecusive( const TSet<FString>& Collapsed, const FString& NodeName, ProjectLauncher::FLaunchProfileTreeNodePtr Node) const;
 	void ReorderNodes(ProjectLauncher::FLaunchProfileTreeNodePtr SourceNode, ProjectLauncher::FLaunchProfileTreeNodePtr TargetNode, EItemDropZone DropZone);
 
+	bool IsHierarchyVisible(ProjectLauncher::FLaunchProfileTreeNodePtr TreeNode) const;
+
 	enum class EMenuType
 	{
 		ProfileExtensions,
@@ -78,6 +80,8 @@ private:
 	TSharedPtr<ProjectLauncher::ILaunchProfileTreeBuilder> TreeBuilder;
 	TSharedPtr<SLaunchProfileTreeView> TreeView;
 	TArray<TSharedRef<ProjectLauncher::FLaunchExtensionInstance>> PendingExtensionDeletes;
+	TSet<FString> PendingCollapseHeadings;
+	FString PendingNewHeadingAdded;
 
 	float SplitterPos = 0.6f;
 	
diff --git a/Engine/Plugins/Developer/ProjectLauncher/Source/ProjectLauncher/Public/Extension/BuildCookRunCommandExtension.h b/Engine/Plugins/Developer/ProjectLauncher/Source/ProjectLauncher/Public/Extension/BuildCookRunCommandExtension.h
--- a/Engine/Plugins/Developer/ProjectLauncher/Source/ProjectLauncher/Public/Extension/BuildCookRunCommandExtension.h
+++ b/Engine/Plugins/Developer/ProjectLauncher/Source/ProjectLauncher/Public/Extension/BuildCookRunCommandExtension.h
@@ -39,11 +39,6 @@ namespace ProjectLauncher
 		 */
 		UE_API FLaunchProfileTreeNode& AddDefaultHeading(FLaunchProfileTreeNode& ProfileTreeNode) const;
 
-		/**
-		 * Create the default heading for this extension, with a bespoke display name
-		 */
-		UE_API FLaunchProfileTreeNode& AddDefaultHeading(FLaunchProfileTreeNode& ProfileTreeNode, FText OverrideDisplayName) const;
-
 
 
 		/**
diff --git a/Engine/Plugins/Developer/ProjectLauncher/Source/ProjectLauncher/Public/ProfileTree/LaunchProfileTreeData.h b/Engine/Plugins/Developer/ProjectLauncher/Source/ProjectLauncher/Public/ProfileTree/LaunchProfileTreeData.h
--- a/Engine/Plugins/Developer/ProjectLauncher/Source/ProjectLauncher/Public/ProfileTree/LaunchProfileTreeData.h
+++ b/Engine/Plugins/Developer/ProjectLauncher/Source/ProjectLauncher/Public/ProfileTree/LaunchProfileTreeData.h
@@ -209,6 +209,11 @@ namespace ProjectLauncher
 		TSet<FLaunchProfileTreeNodePtr> UATCommandHeadings;              // headings that have additional controls for modifying the command
 		bool bRequestTreeRefresh = false;                                // request to refresh tree UI
 		bool bRequestTreeFullRebuild = false;                            // request to fully recreate the profile tree
+
+		/* Constants */
+		UE_API static const TCHAR* GeneralSettingsSectionName;
+		UE_API static const int32 GeneralSettingsSortOrder;
+		UE_API static const int32 DefaultSortOrder;
 	};
 
 
