Make Project Launcher open with last edited profile

#rnx

#Synced-CL 50096176

diff --git a/Engine/Plugins/Developer/ProjectLauncher/Source/ProjectLauncher/Private/Model/ProjectLauncherModel.cpp b/Engine/Plugins/Developer/ProjectLauncher/Source/ProjectLauncher/Private/Model/ProjectLauncherModel.cpp
--- a/Engine/Plugins/Developer/ProjectLauncher/Source/ProjectLauncher/Private/Model/ProjectLauncherModel.cpp
+++ b/Engine/Plugins/Developer/ProjectLauncher/Source/ProjectLauncher/Private/Model/ProjectLauncherModel.cpp
@@ -31,6 +31,8 @@ namespace ProjectLauncher
 	const FText ProjectLauncherText = LOCTEXT("ProjectLauncherTitle", "Project Launcher");
 	bool bUseClassicView = false;
 
+	const TCHAR* LastProfileKey = TEXT("LastProfileId");
+
 	namespace CmdLineType
 	{
 		const TCHAR* Launch = TEXT("Launch");
@@ -80,7 +82,15 @@ namespace ProjectLauncher
 		DefaultBasicLaunchProfile = CreateBasicLaunchProfile();
 		DefaultCustomLaunchProfile = CreateCustomProfile(TEXT("DefaultCustomProfile"));
 
-
+		FString LastProfileId;
+		if (GConfig->GetString(GetConfigSection(), LastProfileKey, LastProfileId, GetConfigIni()))
+		{
+			ILauncherProfilePtr LastProfile = ProfileManager->GetProfile(FGuid(LastProfileId));
+			if (LastProfile.IsValid())
+			{
+				SelectProfile(LastProfile);
+			}
+		}
 	}
 
 
@@ -123,6 +133,8 @@ namespace ProjectLauncher
 				{
 					EnsureValidBuildTarget(SelectedProfile.ToSharedRef(), BuildCookRun);
 				}
+
+				GConfig->SetString(GetConfigSection(), LastProfileKey, *LexToString(SelectedProfile->GetId()), GetConfigIni());
 			}
 
 			ProfileSelectedDelegate.Broadcast(SelectedProfile, PreviousProfile);
diff --git a/Engine/Plugins/Developer/ProjectLauncher/Source/ProjectLauncher/Private/Widgets/Profiles/SCustomLaunchCustomProfileSelector.cpp b/Engine/Plugins/Developer/ProjectLauncher/Source/ProjectLauncher/Private/Widgets/Profiles/SCustomLaunchCustomProfileSelector.cpp
--- a/Engine/Plugins/Developer/ProjectLauncher/Source/ProjectLauncher/Private/Widgets/Profiles/SCustomLaunchCustomProfileSelector.cpp
+++ b/Engine/Plugins/Developer/ProjectLauncher/Source/ProjectLauncher/Private/Widgets/Profiles/SCustomLaunchCustomProfileSelector.cpp
@@ -96,6 +96,7 @@ void SCustomLaunchCustomProfileSelector::Construct(const FArguments& InArgs, con
 	Model->GetProfileManager()->OnProfileRemoved().AddSP(this, &SCustomLaunchCustomProfileSelector::OnCustomProfileRemoved);
 
 	Model->OnProfileSelected().AddSP(this, &SCustomLaunchCustomProfileSelector::OnProfileSelected);
+	OnProfileSelected(Model->GetSelectedProfile(), nullptr);
 }
 END_SLATE_FUNCTION_BUILD_OPTIMIZATION
 
diff --git a/Engine/Plugins/Developer/ProjectLauncher/Source/ProjectLauncher/Private/Widgets/Profiles/SCustomLaunchProfilesPanel.cpp b/Engine/Plugins/Developer/ProjectLauncher/Source/ProjectLauncher/Private/Widgets/Profiles/SCustomLaunchProfilesPanel.cpp
--- a/Engine/Plugins/Developer/ProjectLauncher/Source/ProjectLauncher/Private/Widgets/Profiles/SCustomLaunchProfilesPanel.cpp
+++ b/Engine/Plugins/Developer/ProjectLauncher/Source/ProjectLauncher/Private/Widgets/Profiles/SCustomLaunchProfilesPanel.cpp
@@ -28,10 +28,17 @@
 #include "Framework/Docking/TabManager.h"
 #include "Framework/MultiBox/MultiBoxBuilder.h"
 #include "PlatformInfo.h"
+#include "Misc/ConfigCacheIni.h"
 
 
 #define LOCTEXT_NAMESPACE "SCustomLaunchProfilesPanel"
 
+namespace
+{
+	const TCHAR* WasEditorOpenKey = TEXT("WasEditorOpen");
+	bool bIsEditorOpen = false;
+}
+
 
 SCustomLaunchProfilesPanel::SCustomLaunchProfilesPanel()
 {
@@ -110,15 +117,24 @@ void SCustomLaunchProfilesPanel::Construct(const FArguments& InArgs, const TShar
 	];
 
 	// hide profile editor by default
-	SetProfileEditorVisible(false);
+	bIsEditorOpen = GConfig->GetBoolOrDefault(Model->GetConfigSection(), WasEditorOpenKey, false, Model->GetConfigIni());
+	SetProfileEditorVisible(bIsEditorOpen);
 
 
 	Model->OnProfileSelected().AddSP(this, &SCustomLaunchProfilesPanel::OnProfileSelected);
 
-	// wait one tick then select the Basic Launch profile. this should give enough time for the host device proxy to be created
+	// wait one tick then select the last used profile or in case it's missing the Basic Launch Profile. this should give enough time for the host device proxy to be created
 	ExecuteOnGameThread(UE_SOURCE_LOCATION, [this]()
 	{
-		Model->SelectProfile(Model->GetBasicLaunchProfile());
+		ILauncherProfilePtr SelectedProfile = Model->GetSelectedProfile();
+		if (!SelectedProfile.IsValid())
+		{
+			Model->SelectProfile(Model->GetBasicLaunchProfile());
+		}
+		else
+		{
+			PropertyEditor->SetProfile(SelectedProfile);
+		}
 	});
 }
 END_SLATE_FUNCTION_BUILD_OPTIMIZATION
@@ -900,6 +916,12 @@ void SCustomLaunchProfilesPanel::SetProfileEditorVisible( bool bVisible )
 	
 		ProfileEditorSplitter->SlotAt(0).SetResizable(bVisible);
 		ProfileEditorSplitter->SlotAt(1).SetResizable(bVisible);
+
+		if (bIsEditorOpen != bVisible)
+		{
+			bIsEditorOpen = bVisible;
+			GConfig->SetBool(Model->GetConfigSection(), WasEditorOpenKey, bIsEditorOpen, Model->GetConfigIni());
+		}
 	}
 }
 
