Extensions moved to overlow menu in general settings header

#rnx

#Synced-CL 50118602

diff --git a/Engine/Plugins/Developer/ProjectLauncher/Source/ProjectLauncher/Private/Widgets/Profiles/SCustomLaunchCustomProfileEditor.cpp b/Engine/Plugins/Developer/ProjectLauncher/Source/ProjectLauncher/Private/Widgets/Profiles/SCustomLaunchCustomProfileEditor.cpp
--- a/Engine/Plugins/Developer/ProjectLauncher/Source/ProjectLauncher/Private/Widgets/Profiles/SCustomLaunchCustomProfileEditor.cpp
+++ b/Engine/Plugins/Developer/ProjectLauncher/Source/ProjectLauncher/Private/Widgets/Profiles/SCustomLaunchCustomProfileEditor.cpp
@@ -90,13 +90,15 @@ public:
 
 	ProjectLauncher::FLaunchProfileTreeNodePtr TreeNode;
 	bool bIsPrimaryHeading;
+	bool bIsRootHeading;
 	SCustomLaunchCustomProfileEditor* ProfileEditor;
 
-	void Construct(const FArguments& InArgs, ProjectLauncher::FLaunchProfileTreeNodePtr InTreeNode, SCustomLaunchCustomProfileEditor* InProfileEditor, bool bInIsPrimaryHeading, bool bPinned )
+	void Construct(const FArguments& InArgs, ProjectLauncher::FLaunchProfileTreeNodePtr InTreeNode, SCustomLaunchCustomProfileEditor* InProfileEditor, bool bInIsPrimaryHeading, bool bInIsRootHeading, bool bPinned )
 	{
 		TreeNode = InTreeNode;
 		ProfileEditor = InProfileEditor;
 		bIsPrimaryHeading = bInIsPrimaryHeading;
+		bIsRootHeading = bInIsRootHeading;
 
 		auto GetErrorVisibility = [this]()
 		{
@@ -273,7 +275,7 @@ public:
 			];
 		}
 
-		if (bIsPrimaryHeading && !bPinned)
+		if (bIsRootHeading && !bPinned)
 		{
 			// spacer to push everthing else to the right
 			HeaderBar->AddSlot()
@@ -284,15 +286,15 @@ public:
 
 			// drop-down menu
 			HeaderBar->AddSlot()
-			.HAlign(HAlign_Right)
+			.HAlign(HAlign_Center)
 			.VAlign(VAlign_Center)
-			.Padding(2, 0, 16, 0)
+			.Padding(2, 0, 6, 0)
 			.AutoWidth()
 			[
 				SNew(SComboButton)
 				.ButtonStyle(FAppStyle::Get(), "SimpleButton")
 				.HasDownArrow(false)
-				.OnGetMenuContent( ProfileEditor, &SCustomLaunchCustomProfileEditor::MakeUATCommandSubMenuWidget, TreeNode->UATCommand.ToSharedRef())
+				.OnGetMenuContent( ProfileEditor, &SCustomLaunchCustomProfileEditor::MakeTreeNodeSubMenuWidget, TreeNode)
 				.ButtonContent()
 				[
 					SNew(SImage)
@@ -314,10 +316,10 @@ public:
 
 	virtual FReply OnMouseButtonUp(const FGeometry& Geometry, const FPointerEvent& MouseEvent) override
 	{
-		if (MouseEvent.GetEffectingButton() == EKeys::RightMouseButton && bIsPrimaryHeading)
+		if (MouseEvent.GetEffectingButton() == EKeys::RightMouseButton && bIsRootHeading)
 		{
 			FMenuBuilder MenuBuilder(true, nullptr);
-			ProfileEditor->MakeUATCommandSubMenu(MenuBuilder, TreeNode->UATCommand.ToSharedRef());
+			ProfileEditor->MakeTreeNodeSubMenu(MenuBuilder, TreeNode);
 
 			FSlateApplication::Get().PushMenu( SharedThis(this),
 				FWidgetPath(),
@@ -719,7 +721,8 @@ TSharedRef<ITableRow> SCustomLaunchCustomProfileEditor::OnGenerateWidgetForTreeN
 	else
 	{
 		bool bIsPrimaryHeading = TreeNode->UATCommand.IsValid() && TreeBuilder->GetProfileTree()->UATCommandHeadings.Contains(TreeNode);
-		return SNew(SLaunchProfileCategoryTreeRow, TreeNode, this, bIsPrimaryHeading, bPinned);
+		bool bIsRootHeading = TreeBuilder->GetProfileTree()->Nodes.Contains(TreeNode);
+		return SNew(SLaunchProfileCategoryTreeRow, TreeNode, this, bIsPrimaryHeading, bIsRootHeading, bPinned);
 	}
 }
 
@@ -1192,20 +1195,6 @@ void SCustomLaunchCustomProfileEditor::MakeExtensionMenu( FMenuBuilder& MenuBuil
 
 
 
-TSharedRef<SWidget> SCustomLaunchCustomProfileEditor::MakeExtensionsMenu()
-{
-	// ...construct the extensions menu...
-
-	const bool bShouldCloseWindowAfterMenuSelection = true;
-	const bool bCloseSelfOnly = false;
-	const bool bSearchable = false;
-	FMenuBuilder MenuBuilder(bShouldCloseWindowAfterMenuSelection, nullptr, nullptr, bCloseSelfOnly, &FCoreStyle::Get(), bSearchable );
-
-	MakeMenu(MenuBuilder, EMenuType::ProfileExtensions);
-
-	return MenuBuilder.MakeWidget();
-}
-
 
 TSharedRef<SWidget> SCustomLaunchCustomProfileEditor::MakeCommandsMenu()
 {
@@ -1293,7 +1282,20 @@ void SCustomLaunchCustomProfileEditor::MakeUATCommandSubMenu( FMenuBuilder& Menu
 	}
 }
 
-TSharedRef<SWidget> SCustomLaunchCustomProfileEditor::MakeUATCommandSubMenuWidget(ILauncherProfileUATCommandRef UATCommand)
+
+void SCustomLaunchCustomProfileEditor::MakeTreeNodeSubMenu( FMenuBuilder& MenuBuilder, ProjectLauncher::FLaunchProfileTreeNodePtr TreeNode)
+{
+	if (TreeNode->UATCommand.IsValid() && !Model->CanUseSimplifiedLayout(CurrentProfile.ToSharedRef()))
+	{
+		MakeUATCommandSubMenu(MenuBuilder, TreeNode->UATCommand.ToSharedRef());
+	}
+	else
+	{
+		MakeMenu(MenuBuilder, EMenuType::ProfileExtensions);
+	}
+}
+
+TSharedRef<SWidget> SCustomLaunchCustomProfileEditor::MakeTreeNodeSubMenuWidget(ProjectLauncher::FLaunchProfileTreeNodePtr TreeNode)
 {
 	const bool bShouldCloseWindowAfterMenuSelection = true;
 	const bool bCloseSelfOnly = true;
@@ -1301,7 +1303,7 @@ TSharedRef<SWidget> SCustomLaunchCustomProfileEditor::MakeUATCommandSubMenuWidge
 	const bool bInOpenSubMenuOnClick = false;
 	FMenuBuilder MenuBuilder(bShouldCloseWindowAfterMenuSelection, nullptr, nullptr, bCloseSelfOnly, &FCoreStyle::Get(), bSearchable );
 
-	MakeUATCommandSubMenu(MenuBuilder, UATCommand);
+	MakeTreeNodeSubMenu(MenuBuilder, TreeNode);
 
 	return MenuBuilder.MakeWidget();
 }
diff --git a/Engine/Plugins/Developer/ProjectLauncher/Source/ProjectLauncher/Private/Widgets/Profiles/SCustomLaunchCustomProfileEditor.h b/Engine/Plugins/Developer/ProjectLauncher/Source/ProjectLauncher/Private/Widgets/Profiles/SCustomLaunchCustomProfileEditor.h
--- a/Engine/Plugins/Developer/ProjectLauncher/Source/ProjectLauncher/Private/Widgets/Profiles/SCustomLaunchCustomProfileEditor.h
+++ b/Engine/Plugins/Developer/ProjectLauncher/Source/ProjectLauncher/Private/Widgets/Profiles/SCustomLaunchCustomProfileEditor.h
@@ -37,7 +37,6 @@ public:
 	void SetProfile( const ILauncherProfilePtr& Profile );
 	ILauncherProfilePtr GetProfile() const { return CurrentProfile; }
 
-	TSharedRef<SWidget> MakeExtensionsMenu();
 	TSharedRef<SWidget> MakeCommandsMenu();
 
 	void RebuildTree();
@@ -70,7 +69,9 @@ private:
 
 	void MakeUATCommandSubMenu( FMenuBuilder& MenuBuilder, ILauncherProfileUATCommandRef UATCommand );
 	void MakeExtensionMenu( FMenuBuilder& MenuBuilder, EMenuType MenuType, TSharedRef<ProjectLauncher::FLaunchExtension> Extension, ILauncherProfileBuildCookRunPtr BuildCookRun = nullptr );
-	TSharedRef<SWidget> MakeUATCommandSubMenuWidget(ILauncherProfileUATCommandRef UATCommand);
+
+	void MakeTreeNodeSubMenu( FMenuBuilder& MenuBuilder, ProjectLauncher::FLaunchProfileTreeNodePtr TreeNode);
+	TSharedRef<SWidget> MakeTreeNodeSubMenuWidget(ProjectLauncher::FLaunchProfileTreeNodePtr TreeNode);
 
 	TSharedPtr<ProjectLauncher::FLaunchExtensionInstance> GetExtensionForUATCommand(ILauncherProfileUATCommandRef UATCommand) const;
 	FSlateIcon GetIconForUATCommand(ILauncherProfileUATCommandPtr UATCommand) const;
diff --git a/Engine/Plugins/Developer/ProjectLauncher/Source/ProjectLauncher/Private/Widgets/Profiles/SCustomLaunchProfilesPanel.cpp b/Engine/Plugins/Developer/ProjectLauncher/Source/ProjectLauncher/Private/Widgets/Profiles/SCustomLaunchProfilesPanel.cpp
--- a/Engine/Plugins/Developer/ProjectLauncher/Source/ProjectLauncher/Private/Widgets/Profiles/SCustomLaunchProfilesPanel.cpp
+++ b/Engine/Plugins/Developer/ProjectLauncher/Source/ProjectLauncher/Private/Widgets/Profiles/SCustomLaunchProfilesPanel.cpp
@@ -361,42 +361,6 @@ TSharedRef<SWidget> SCustomLaunchProfilesPanel::CreateProfileEditorWidget()
 				.Text(LOCTEXT("TasksLabel", "Launch Tasks"))
 			]
 
-			// extensions menu
-			+SHorizontalBox::Slot()
-			.AutoWidth()
-			.VAlign(VAlign_Center)
-			.HAlign(HAlign_Right)
-			.Padding(4,2)
-			[
-				SNew(SComboButton)
-				.ComboButtonStyle(FAppStyle::Get(), "SimpleComboButton")
-				.ToolTipText(LOCTEXT("ExtensionsToolTip", "Select launcher extensions."))
-				.OnGetMenuContent(PropertyEditor.ToSharedRef(), &SCustomLaunchCustomProfileEditor::MakeExtensionsMenu)
-				.MenuPlacement(MenuPlacement_ComboBoxRight)
-				.ButtonContent()
-				[
-					SNew(SHorizontalBox)
-					+SHorizontalBox::Slot()
-					.AutoWidth()
-					.VAlign(VAlign_Center)
-					.Padding(1,3)
-					[
-						SNew(SImage)
-						.Image(FProjectLauncherStyle::Get().GetBrush("Icons.Plugins"))
-						.ColorAndOpacity(FSlateColor::UseForeground())
-					]
-					+SHorizontalBox::Slot()
-					.AutoWidth()
-					.VAlign(VAlign_Center)
-					.Padding(1,3)
-					[
-						SNew(STextBlock)
-						.Text(LOCTEXT("ExtensionsLabel", "Extensions"))
-						.ColorAndOpacity(FSlateColor::UseForeground())
-					]
-				]
-			]
-
 			// new task button
 			+SHorizontalBox::Slot()
 			.AutoWidth()
