Minor refactor for command line extensions
- one Type per extension to make things easier
- parameter query functions moved into factory + renamed CmdLine to make things clear
- added 'default' cmdline parameters which are automatically added to all command line fields when the extension is added

#Synced-CL 50586249

diff --git a/Engine/Plugins/Developer/ProjectLauncher/Source/CommonLaunchExtensions/Private/Insights/InsightsLaunchExtension.cpp b/Engine/Plugins/Developer/ProjectLauncher/Source/CommonLaunchExtensions/Private/Insights/InsightsLaunchExtension.cpp
--- a/Engine/Plugins/Developer/ProjectLauncher/Source/CommonLaunchExtensions/Private/Insights/InsightsLaunchExtension.cpp
+++ b/Engine/Plugins/Developer/ProjectLauncher/Source/CommonLaunchExtensions/Private/Insights/InsightsLaunchExtension.cpp
@@ -18,10 +18,6 @@ TSharedRef<ProjectLauncher::FCmdLineParametersExtension> FInsightsLaunchExtensio
 	return MakeShared<FInsightsCommandLineExtension>(InArgs);
 }
 
-bool FInsightsLaunchExtensionInstance::SupportsCmdLineType( const FString& Type ) const
-{
-	return ICmdLineParametersExtensionFactory::DefaultSupportsType(Type);
-}
 
 
 
@@ -30,16 +26,14 @@ FInsightsLaunchExtensionInstance::FInsightsCommandLineExtension::FInsightsComman
 {
 }
 
-bool FInsightsLaunchExtensionInstance::FInsightsCommandLineExtension::GetExtensionParameters( TArray<FString>& OutParameters ) const
+void FInsightsLaunchExtensionInstance::GetCmdLineParameters( TArray<FString>& OutParameters ) const
 {
 	OutParameters.Add(HostParam);
 	OutParameters.Add(FileParam);
 	OutParameters.Add(StatNamedEventsParam);
-
-	return true;
 }
 
-FText FInsightsLaunchExtensionInstance::FInsightsCommandLineExtension::GetExtensionParameterDisplayName( const FString& InParameter ) const
+FText FInsightsLaunchExtensionInstance::GetCmdLineParameterDisplayName( const FString& InParameter ) const
 {
 	if (InParameter == HostParam)
 	{
@@ -54,7 +48,7 @@ FText FInsightsLaunchExtensionInstance::FInsightsCommandLineExtension::GetExtens
 		return LOCTEXT("TraceNamedEventsParam", "Capture named events");
 	}
 
-	return ProjectLauncher::FCmdLineParametersExtension::GetExtensionParameterDisplayName(InParameter);
+	return ProjectLauncher::ICmdLineParametersExtensionFactory::GetCmdLineParameterDisplayName(InParameter);
 
 }
 
diff --git a/Engine/Plugins/Developer/ProjectLauncher/Source/CommonLaunchExtensions/Private/Insights/InsightsLaunchExtension.h b/Engine/Plugins/Developer/ProjectLauncher/Source/CommonLaunchExtensions/Private/Insights/InsightsLaunchExtension.h
--- a/Engine/Plugins/Developer/ProjectLauncher/Source/CommonLaunchExtensions/Private/Insights/InsightsLaunchExtension.h
+++ b/Engine/Plugins/Developer/ProjectLauncher/Source/CommonLaunchExtensions/Private/Insights/InsightsLaunchExtension.h
@@ -16,7 +16,8 @@ public:
 
 	// ... ICmdLineParametersExtensionFactory
 	virtual TSharedRef<ProjectLauncher::FCmdLineParametersExtension> Create( const ProjectLauncher::FCmdLineParametersExtension::FArgs& InArgs ) override;
-	virtual bool SupportsCmdLineType( const FString& Type ) const override;
+	virtual void GetCmdLineParameters( TArray<FString>& OutParameters ) const override;
+	virtual FText GetCmdLineParameterDisplayName( const FString& InParameter ) const override;
 
 private:
 	class FInsightsCommandLineExtension : public ProjectLauncher::FCmdLineParametersExtension
@@ -24,9 +25,7 @@ private:
 	public:
 		FInsightsCommandLineExtension( const ProjectLauncher::FCmdLineParametersExtension::FArgs& InArgs );
 
-		virtual bool GetExtensionParameters( TArray<FString>& OutParameters ) const override;
-		virtual FText GetExtensionParameterDisplayName( const FString& InParameter ) const override;
-	
+
 		virtual void CustomizeParametersSubmenu( FMenuBuilder& MenuBuilder ) override;
 
 		void CacheTraceChannels();
diff --git a/Engine/Plugins/Developer/ProjectLauncher/Source/ProjectLauncher/Private/Extension/CmdLineParametersExtension.cpp b/Engine/Plugins/Developer/ProjectLauncher/Source/ProjectLauncher/Private/Extension/CmdLineParametersExtension.cpp
--- a/Engine/Plugins/Developer/ProjectLauncher/Source/ProjectLauncher/Private/Extension/CmdLineParametersExtension.cpp
+++ b/Engine/Plugins/Developer/ProjectLauncher/Source/ProjectLauncher/Private/Extension/CmdLineParametersExtension.cpp
@@ -29,16 +29,6 @@ namespace ProjectLauncher
 		Profile->OnProjectChanged().RemoveAll(this);
 	}
 
-	FText FCmdLineParametersExtension::GetExtensionParameterDisplayName( const FString& InParameter ) const
-	{
-		return FText::FromString(InParameter);
-	}
-
-	bool FCmdLineParametersExtension::GetExtensionParameters( TArray<FString>& OutParameters ) const
-	{
-		return false;
-	}
-
 	void FCmdLineParametersExtension::SetCommandLine( const FString& InCommandLine )
 	{
 		SetCommandLineDelegate(InCommandLine);
@@ -51,55 +41,13 @@ namespace ProjectLauncher
 	}
 
 
-	void FCmdLineParametersExtension::MakeCommandLineSubmenu( FMenuBuilder& MenuBuilder )
-	{
-		// create the menu entry
-		auto ToggleParameter = [this]( FString Parameter )
-		{
-			FString CommandLine = GetCommandLine();
-			if (CmdLineUtils::IsParameterUsed(CommandLine, Parameter))
-			{
-				CmdLineUtils::RemoveParameter(CommandLine, Parameter);
-			}
-			else
-			{
-				CmdLineUtils::AddParameter(CommandLine, Parameter);
-			}
-			SetCommandLine(CommandLine);
-		};
-
-		TArray<FString> Parameters;
-		if (GetExtensionParameters(Parameters) && Parameters.Num() > 0)
-		{
-			for (const FString& Parameter : Parameters)
-			{
-				FText ParameterDisplayName = GetExtensionParameterDisplayName(Parameter);
-				FText ParameterToolTip = FText::FromString(Parameter);
-
-				MenuBuilder.AddMenuEntry( 
-					ParameterDisplayName,
-					ParameterToolTip,
-					FSlateIcon(),
-					FUIAction(
-						FExecuteAction::CreateLambda(ToggleParameter, Parameter),
-						FCanExecuteAction(),
-						FIsActionChecked::CreateLambda( [this, Parameter]() { return CmdLineUtils::IsParameterUsed(GetCommandLine(), Parameter); } )
-					),
-					NAME_None,
-					EUserInterfaceActionType::Check);
-			}
-		}
-
-		CustomizeParametersSubmenu(MenuBuilder);
-	}
-
-
-
-
 	static TMap<uint64,TArray<TSharedRef<FCmdLineParametersExtension>>> GCmdLineParametersExtensions;
 
 	void ICmdLineParametersExtensionFactory::MakeCmdLineParametersSubmenu(FMenuBuilder& MenuBuilder, uint64 InstanceID, const FString& TypeName, TSharedRef<FModel> Model, TSharedRef<FLaunchExtensionInstance> ExtensionInstance, const FLaunchProfileTreeNode::FSetString& SetValue, const FLaunchProfileTreeNode::FGetString& GetValue )
 	{
+		ProjectLauncher::ICmdLineParametersExtensionFactory* CmdLineParametersFactory = ExtensionInstance->AsCmdLineParametersFactory();
+		check(CmdLineParametersFactory); // should have been verified before calling
+
 		// get the command line instances for this field
 		if (!GCmdLineParametersExtensions.Contains(InstanceID))
 		{
@@ -120,9 +68,6 @@ namespace ProjectLauncher
 		else
 		{
 			// create & cache a new extension instance
-			ProjectLauncher::ICmdLineParametersExtensionFactory* CmdLineParametersFactory = ExtensionInstance->AsCmdLineParametersFactory();
-			check(CmdLineParametersFactory); // should have been verified before calling
-
 			FCmdLineParametersExtension::FArgs Args
 			{
 				.Model = Model,
@@ -137,9 +82,56 @@ namespace ProjectLauncher
 		}
 
 		// create the menu
-		CmdLineExtension->MakeCommandLineSubmenu(MenuBuilder);
+		TArray<FString> Parameters;
+		CmdLineParametersFactory->GetCmdLineParameters(Parameters);
+		if (Parameters.Num() > 0)
+		{
+			auto ToggleParameter = [CmdLineExtension]( FString Parameter )
+			{
+				FString CommandLine = CmdLineExtension->GetCommandLine();
+				if (CmdLineUtils::IsParameterUsed(CommandLine, Parameter))
+				{
+					CmdLineUtils::RemoveParameter(CommandLine, Parameter);
+				}
+				else
+				{
+					CmdLineUtils::AddParameter(CommandLine, Parameter);
+				}
+				CmdLineExtension->SetCommandLine(CommandLine);
+			};
+
+
+			for (const FString& Parameter : Parameters)
+			{
+				FText ParameterDisplayName = CmdLineParametersFactory->GetCmdLineParameterDisplayName(Parameter);
+				FText ParameterToolTip = FText::FromString(Parameter);
+
+				MenuBuilder.AddMenuEntry( 
+					ParameterDisplayName,
+					ParameterToolTip,
+					FSlateIcon(),
+					FUIAction(
+						FExecuteAction::CreateLambda(ToggleParameter, Parameter),
+						FCanExecuteAction(),
+						FIsActionChecked::CreateLambda( [CmdLineExtension, Parameter]() { return CmdLineUtils::IsParameterUsed(CmdLineExtension->GetCommandLine(), Parameter); } )
+					),
+					NAME_None,
+					EUserInterfaceActionType::Check);
+			}
+		}
+
+		CmdLineExtension->CustomizeParametersSubmenu(MenuBuilder);
+	}
+
+	TSharedRef<FCmdLineParametersExtension> ICmdLineParametersExtensionFactory::Create( const FCmdLineParametersExtension::FArgs& InArgs )
+	{
+		return MakeShared<ProjectLauncher::FCmdLineParametersExtension>(InArgs);
 	}
 
+	FText ICmdLineParametersExtensionFactory::GetCmdLineParameterDisplayName( const FString& InParameter ) const
+	{
+		return FText::FromString(InParameter);
+	}
 
 	void ICmdLineParametersExtensionFactory::DestroyCmdLineParametersExtensions(TSharedPtr<FLaunchExtensionInstance> Instance)
 	{
@@ -164,13 +156,6 @@ namespace ProjectLauncher
 			GCmdLineParametersExtensions.Remove(EmptyInstanceID);
 		}
 	}
-
-	bool ICmdLineParametersExtensionFactory::DefaultSupportsType( const FString& Type )
-	{
-		return Type.Equals(ProjectLauncher::CmdLineType::Launch);
-	}
-
-
 }
 
 #undef LOCTEXT_NAMESPACE
diff --git a/Engine/Plugins/Developer/ProjectLauncher/Source/ProjectLauncher/Private/ProfileTree/LaunchProfileTreeData.cpp b/Engine/Plugins/Developer/ProjectLauncher/Source/ProjectLauncher/Private/ProfileTree/LaunchProfileTreeData.cpp
--- a/Engine/Plugins/Developer/ProjectLauncher/Source/ProjectLauncher/Private/ProfileTree/LaunchProfileTreeData.cpp
+++ b/Engine/Plugins/Developer/ProjectLauncher/Source/ProjectLauncher/Private/ProfileTree/LaunchProfileTreeData.cpp
@@ -4,6 +4,7 @@
 #include "ProfileTree/ILaunchProfileTreeBuilder.h"
 #include "Extension/LaunchExtension.h"
 #include "Extension/CmdLineParametersExtension.h"
+#include "Utils/LauncherCmdLineUtils.h"
 #include "Styling/ProjectLauncherStyle.h"
 #include "Styling/AppStyle.h"
 #include "Styling/StyleColors.h"
@@ -213,6 +214,69 @@ namespace ProjectLauncher
 		bRequestTreeFullRebuild = true;
 	}
 
+	void FLaunchProfileTreeData::OnLaunchExtensionToggle( const TSharedRef<FLaunchExtensionInstance>& ExtensionInstance, bool bAdded )
+	{
+		ProjectLauncher::ICmdLineParametersExtensionFactory* CmdLineFactory = ExtensionInstance->AsCmdLineParametersFactory();
+		if (CmdLineFactory)
+		{
+			FString TypeName = CmdLineFactory->GetCmdLineType();
+
+			// toggle the default command line parameters for all matching command line properties
+			for ( const FCommandLineField& CommandLineField : AllCommandLineFields )
+			{
+				if (CommandLineField.TypeName == TypeName)
+				{
+					// get the relevant parameters for this extension
+					static const bool bOnlyRemoveDefaultCmdLineParameters = false;
+
+					TArray<FString> Parameters;
+					if (bAdded || bOnlyRemoveDefaultCmdLineParameters)
+					{
+						CmdLineFactory->GetDefaultCmdLineParameters( Parameters );
+					}
+					else
+					{
+						CmdLineFactory->GetCmdLineParameters( Parameters );
+					}
+
+					// toggle them all
+					if (Parameters.Num() > 0)
+					{
+						FString CommandLine = CommandLineField.GetValue();
+						for (const FString& Parameter : Parameters)
+						{
+							if (bAdded)
+							{
+								CmdLineUtils::AddParameter(CommandLine, Parameter);
+							}
+							else
+							{
+								CmdLineUtils::RemoveParameter(CommandLine, Parameter);
+							}
+						}
+						CommandLineField.SetValue(CommandLine);
+					}
+				}
+			}
+		}
+
+		if (bAdded)
+		{
+			RequestFullTreeRebuild();
+		}
+	}
+
+	void FLaunchProfileTreeData::OnCommandLineFieldCreated( FLaunchProfileTreeNode& InNode, const FLaunchProfileTreeNode::FStringCallbacks& InStringCallbacks, const FString& InTypeName )
+	{
+		// store the command line callback
+		AllCommandLineFields.Emplace( FLaunchProfileTreeData::FCommandLineField 
+		{ 
+			.GetValue = InStringCallbacks.GetValue, 
+			.SetValue = InStringCallbacks.SetValue,
+			.TypeName = InTypeName,
+		});
+	}
+
 	void FLaunchProfileTreeData::OnPropertyChanged()
 	{
 		for (TSharedPtr<FLaunchExtensionInstance> ExtensionInstance : ExtensionInstances)
@@ -645,6 +709,7 @@ namespace ProjectLauncher
 
 		check(StringCallbacks.GetValue);
 		check(StringCallbacks.SetValue);
+		TreeData->OnCommandLineFieldCreated( *this, StringCallbacks, InType );
 
 		FCallbacks WidgetCallbacks
 		{
@@ -687,7 +752,7 @@ namespace ProjectLauncher
 			for (TSharedPtr<FLaunchExtensionInstance> ExtensionInstance : TreeData->ExtensionInstances)
 			{
 				ProjectLauncher::ICmdLineParametersExtensionFactory* CmdLineParametersFactory = ExtensionInstance->AsCmdLineParametersFactory();
-				if (CmdLineParametersFactory == nullptr || !CmdLineParametersFactory->SupportsCmdLineType(TypeName))
+				if (CmdLineParametersFactory == nullptr || CmdLineParametersFactory->GetCmdLineType() != TypeName)
 				{
 					continue;
 				}
diff --git a/Engine/Plugins/Developer/ProjectLauncher/Source/ProjectLauncher/Private/Widgets/Profiles/SCustomLaunchCustomProfileEditor.cpp b/Engine/Plugins/Developer/ProjectLauncher/Source/ProjectLauncher/Private/Widgets/Profiles/SCustomLaunchCustomProfileEditor.cpp
--- a/Engine/Plugins/Developer/ProjectLauncher/Source/ProjectLauncher/Private/Widgets/Profiles/SCustomLaunchCustomProfileEditor.cpp
+++ b/Engine/Plugins/Developer/ProjectLauncher/Source/ProjectLauncher/Private/Widgets/Profiles/SCustomLaunchCustomProfileEditor.cpp
@@ -1100,7 +1100,7 @@ void SCustomLaunchCustomProfileEditor::MakeExtensionMenu( FMenuBuilder& MenuBuil
 					TSharedPtr<ProjectLauncher::FLaunchExtensionInstance> Instance = Extension->AsUATCommandFactory()->CreateUATCommandExtensionInstance(CurrentProfile.ToSharedRef(), Model.ToSharedRef(), TreeBuilder->GetProfileTree().Get());
 					if (Instance.IsValid())
 					{
-						TreeBuilder->GetProfileTree()->RequestFullTreeRebuild();
+						TreeBuilder->GetProfileTree()->OnLaunchExtensionToggle(Instance.ToSharedRef(), true);
 						PendingNewHeadingAdded = Instance->GetUATCommand()->GetInternalName();
 					}
 				})
@@ -1113,6 +1113,8 @@ void SCustomLaunchCustomProfileEditor::MakeExtensionMenu( FMenuBuilder& MenuBuil
 		{
 			if (ExtensionInstance.IsValid())
 			{
+				TreeBuilder->GetProfileTree()->OnLaunchExtensionToggle(ExtensionInstance.ToSharedRef(), false);
+
 				// defer the delete until the next tick
 				PendingExtensionDeletes.AddUnique(ExtensionInstance.ToSharedRef());
 			}
@@ -1121,8 +1123,7 @@ void SCustomLaunchCustomProfileEditor::MakeExtensionMenu( FMenuBuilder& MenuBuil
 				TSharedPtr<ProjectLauncher::FLaunchExtensionInstance> Instance = ProjectLauncher::FLaunchExtension::CreateExtensionInstance(Extension, CurrentProfile.ToSharedRef(), Model.ToSharedRef(), TreeBuilder->GetProfileTree().Get());
 				if (Instance.IsValid())
 				{
-					TreeBuilder->GetProfileTree()->RequestFullTreeRebuild();
-
+					TreeBuilder->GetProfileTree()->OnLaunchExtensionToggle(Instance.ToSharedRef(), true);
 					TreeBuilder->OnPropertyChanged();
 					CurrentProfile->RefreshCustomWarningsAndErrors();
 				}
@@ -1260,6 +1261,10 @@ void SCustomLaunchCustomProfileEditor::MakeUATCommandSubMenu( FMenuBuilder& Menu
 				{
 					// see if there is an extension associated with this UAT command
 					TSharedPtr<ProjectLauncher::FLaunchExtensionInstance> ExtensionInstance = GetExtensionForUATCommand(UATCommand);
+					if (ExtensionInstance.IsValid())
+					{
+						TreeData->OnLaunchExtensionToggle(ExtensionInstance.ToSharedRef(), false);
+					}
 
 					TreeData->Profile->RemoveUATCommand(UATCommand->GetInternalName()); 
 
diff --git a/Engine/Plugins/Developer/ProjectLauncher/Source/ProjectLauncher/Public/Extension/CmdLineParametersExtension.h b/Engine/Plugins/Developer/ProjectLauncher/Source/ProjectLauncher/Public/Extension/CmdLineParametersExtension.h
--- a/Engine/Plugins/Developer/ProjectLauncher/Source/ProjectLauncher/Public/Extension/CmdLineParametersExtension.h
+++ b/Engine/Plugins/Developer/ProjectLauncher/Source/ProjectLauncher/Public/Extension/CmdLineParametersExtension.h
@@ -9,6 +9,14 @@
 
 namespace ProjectLauncher
 {
+	/* common command line types */
+	namespace CmdLineType
+	{
+		extern UE_API const TCHAR* Launch;
+		extern UE_API const TCHAR* UAT;
+	}
+
+
 	/**
 	 * Base for a command line extension
 	 */
@@ -27,21 +35,6 @@ namespace ProjectLauncher
 		UE_API FCmdLineParametersExtension( const FArgs& InArgs );
 		UE_API virtual ~FCmdLineParametersExtension();
 
-		/**
-		 * Returns the parameters that this extension provides. They will be added to the submenu.
-		 */
-		UE_API virtual bool GetExtensionParameters( TArray<FString>& OutParameters ) const;
-
-		/**
-		 * Returns the user-facing name for the given parameter. It will default to the parameter itself.
-		 */
-		UE_API virtual FText GetExtensionParameterDisplayName( const FString& InParameter ) const;
-
-		/**
-		 * Creates the submenu for picking command line parameters
-		 */
-		UE_API void MakeCommandLineSubmenu( FMenuBuilder& MenuBuilder );
-
 		/** 
 		 * Get the extension we belong to
 		 */
@@ -91,6 +84,8 @@ namespace ProjectLauncher
 		FLaunchProfileTreeNode::FSetString SetCommandLineDelegate;
 		FLaunchProfileTreeNode::FGetString GetCommandLineDelegate;
 		FString TypeName;
+
+		friend class ICmdLineParametersExtensionFactory;
 	};
 
 
@@ -103,14 +98,31 @@ namespace ProjectLauncher
 		virtual ~ICmdLineParametersExtensionFactory() = default;
 
 		/**
-		 * A launch extension must override this to create a command line extension
+		 * Create an instance of a command line extension. A launch extension can override this to customize the command line menu
 		 */
-		virtual TSharedRef<FCmdLineParametersExtension> Create( const FCmdLineParametersExtension::FArgs& InArgs ) = 0;
+		UE_API virtual TSharedRef<FCmdLineParametersExtension> Create( const FCmdLineParametersExtension::FArgs& InArgs );
 
 		/**
-		 * Whether the extension supports the given type. Either a custom string or one of ProjectLauncher::CmdLineType::
+		 * Which commandline type the extension supports. Either a custom string or one of ProjectLauncher::CmdLineType::
 		 */
-		virtual bool SupportsCmdLineType( const FString& Type ) const = 0;
+		virtual FString GetCmdLineType() const { return CmdLineType::Launch; }
+
+		/**
+		 * Returns the user-facing name for the given parameter. It will default to the parameter itself.
+		 */
+		UE_API virtual FText GetCmdLineParameterDisplayName( const FString& InParameter ) const;
+
+		/**
+		 * Returns the parameters that this extension provides. They will be added to the submenu.
+		 */
+		virtual void GetCmdLineParameters( TArray<FString>& OutParameters ) const = 0;
+
+		/**
+		 * Returns the parameters that will be automatically added to the command line when this extension is enabled.
+		 * Not available for 'always created' extensions
+		 */
+		virtual void GetDefaultCmdLineParameters( TArray<FString>& OutParameters ) const {}
+
 
 
 		/**
@@ -123,10 +135,6 @@ namespace ProjectLauncher
 		 */
 		static UE_API void DestroyCmdLineParametersExtensions(TSharedPtr<FLaunchExtensionInstance> Instance);
 
-		/**
-		 * Helper function for interface implementers - returns True for Launch command lines (the most common case)
-		 */
-		static UE_API bool DefaultSupportsType( const FString& Type );
 	};
 };
 
diff --git a/Engine/Plugins/Developer/ProjectLauncher/Source/ProjectLauncher/Public/Model/ProjectLauncherModel.h b/Engine/Plugins/Developer/ProjectLauncher/Source/ProjectLauncher/Public/Model/ProjectLauncherModel.h
--- a/Engine/Plugins/Developer/ProjectLauncher/Source/ProjectLauncher/Public/Model/ProjectLauncherModel.h
+++ b/Engine/Plugins/Developer/ProjectLauncher/Source/ProjectLauncher/Public/Model/ProjectLauncherModel.h
@@ -129,8 +129,8 @@ namespace ProjectLauncher
 	/* common command line types */
 	namespace CmdLineType
 	{
-		extern const TCHAR* Launch;
-		extern const TCHAR* UAT;
+		extern UE_API const TCHAR* Launch;
+		extern UE_API const TCHAR* UAT;
 	}
 
 
diff --git a/Engine/Plugins/Developer/ProjectLauncher/Source/ProjectLauncher/Public/ProfileTree/LaunchProfileTreeData.h b/Engine/Plugins/Developer/ProjectLauncher/Source/ProjectLauncher/Public/ProfileTree/LaunchProfileTreeData.h
--- a/Engine/Plugins/Developer/ProjectLauncher/Source/ProjectLauncher/Public/ProfileTree/LaunchProfileTreeData.h
+++ b/Engine/Plugins/Developer/ProjectLauncher/Source/ProjectLauncher/Public/ProfileTree/LaunchProfileTreeData.h
@@ -197,6 +197,8 @@ namespace ProjectLauncher
 		UE_API void OnPropertyChanged();
 		UE_API void RequestTreeRefresh();
 		UE_API void RequestFullTreeRebuild();
+		UE_API void OnCommandLineFieldCreated( FLaunchProfileTreeNode& InNode, const FLaunchProfileTreeNode::FStringCallbacks& InStringCallbacks, const FString& InTypeName );
+		UE_API void OnLaunchExtensionToggle( const TSharedRef<class FLaunchExtensionInstance>& ExtensionInstance, bool bAdded );
 
 		/* Member variables */
 		ILauncherProfilePtr Profile;                                     // associated profile
@@ -214,6 +216,15 @@ namespace ProjectLauncher
 		UE_API static const TCHAR* GeneralSettingsSectionName;
 		UE_API static const int32 GeneralSettingsSortOrder;
 		UE_API static const int32 DefaultSortOrder;
+
+	private:
+		struct FCommandLineField
+		{
+			FLaunchProfileTreeNode::FGetString GetValue = nullptr;
+			FLaunchProfileTreeNode::FSetString SetValue = nullptr;
+			FString TypeName;
+		};
+		TArray<FCommandLineField> AllCommandLineFields;
 	};
 
 
