Move Project Launcher BuildSync extension into the public engine folder

#rnx

#Synced-CL 50704017

diff --git a/Engine/Plugins/Developer/ProjectLauncher/Source/CommonLaunchExtensions/CommonLaunchExtensions.Build.cs b/Engine/Plugins/Developer/ProjectLauncher/Source/CommonLaunchExtensions/CommonLaunchExtensions.Build.cs
--- a/Engine/Plugins/Developer/ProjectLauncher/Source/CommonLaunchExtensions/CommonLaunchExtensions.Build.cs
+++ b/Engine/Plugins/Developer/ProjectLauncher/Source/CommonLaunchExtensions/CommonLaunchExtensions.Build.cs
@@ -15,12 +15,15 @@ public class CommonLaunchExtensions : ModuleRules
 			"SlateCore",
 			"Sockets",
 			"TraceLog",
+			"Projects",
 			"ProjectLauncher",
 			"DeveloperSettings",
 			"ToolWidgets",
+			"Zen",
+			"HTTP",
+			"InputCore",
 		]);
 
-
 		if (Target.bBuildEditor)
 		{
 			PrivateDependencyModuleNames.AddRange(
@@ -29,5 +32,28 @@ public class CommonLaunchExtensions : ModuleRules
 				"Engine"
 			]);
 		}
+
+		if (Target.Platform == UnrealTargetPlatform.Win64)
+		{
+			RuntimeDependencies.Add("$(EngineDir)/Binaries/Win64/crashpad_handler.exe");
+			RuntimeDependencies.Add("$(EngineDir)/Binaries/Win64/zen.exe");
+			RuntimeDependencies.Add("$(EngineDir)/Binaries/Win64/zen.pdb");
+			RuntimeDependencies.Add("$(EngineDir)/Binaries/Win64/zenserver.exe");
+			RuntimeDependencies.Add("$(EngineDir)/Binaries/Win64/zenserver.pdb");
+			RuntimeDependencies.Add("$(EngineDir)/Binaries/DotNET/OidcToken/win-x64/OidcToken.exe");
+		}
+		else if (Target.Platform == UnrealTargetPlatform.Linux)
+		{
+			RuntimeDependencies.Add("$(EngineDir)/Binaries/Linux/zen");
+			RuntimeDependencies.Add("$(EngineDir)/Binaries/Linux/zenserver");
+			RuntimeDependencies.Add("$(EngineDir)/Binaries/DotNET/OidcToken/linux-x64/OidcToken");
+		}
+		else if (Target.Platform == UnrealTargetPlatform.Mac)
+		{
+			RuntimeDependencies.Add("$(EngineDir)/Binaries/Mac/crashpad_handler");
+			RuntimeDependencies.Add("$(EngineDir)/Binaries/Mac/zen");
+			RuntimeDependencies.Add("$(EngineDir)/Binaries/Mac/zenserver");
+			RuntimeDependencies.Add("$(EngineDir)/Binaries/DotNET/OidcToken/osx-x64/OidcToken");
+		}
 	}
 }
diff --git a/Engine/Plugins/Developer/ProjectLauncher/Source/CommonLaunchExtensions/Private/BuildSync/BuildInfoHelper.cpp b/Engine/Plugins/Developer/ProjectLauncher/Source/CommonLaunchExtensions/Private/BuildSync/BuildInfoHelper.cpp
new file mode 100644
--- /dev/null
+++ b/Engine/Plugins/Developer/ProjectLauncher/Source/CommonLaunchExtensions/Private/BuildSync/BuildInfoHelper.cpp
@@ -0,0 +1,437 @@
+// Copyright Epic Games, Inc. All Rights Reserved.
+
+#include "BuildSync/BuildInfoHelper.h"
+#include "Misc/EngineVersion.h"
+#include "Misc/DataDrivenPlatformInfoRegistry.h"
+#include "Internationalization/Internationalization.h"
+
+#define LOCTEXT_NAMESPACE "FBuildSyncLaunchExtensionInstance"
+
+
+
+const TCHAR* FBuildInfoHelper::DefaultBuildType = TEXT("staged-build");
+
+
+FBuildInfoHelper::FBuildInfoHelper()
+	: bBuildListRetrieverInitComplete(false)
+{
+	BuildListRetriever = MakeShared<UE::Zen::Build::FBuildListRetriever>();
+	BackendRetriever = MakeShared<FBackendRetriever>();
+	UGSBuildInfoRetriever = MakeShared<FUGSBuildInfoRetriever>();
+
+	BuildType = DefaultBuildType;
+
+
+	for ( const auto& Pair : FDataDrivenPlatformInfoRegistry::GetAllPlatformInfos())
+	{
+		const FDataDrivenPlatformInfo& DDPI = Pair.Value;
+		if (!DDPI.bIsFakePlatform)
+		{
+			KnownPlatformNames.Add(DDPI.IniPlatformName.ToString().ToLower());
+			KnownPlatformNames.Add(DDPI.UBTPlatformName.ToString().ToLower());
+		}
+	}
+}
+
+FBuildInfoHelper::~FBuildInfoHelper()
+{
+	BuildListRetriever->SetOnReadyForQuery(nullptr);
+	BuildListRetriever.Reset();
+}
+
+void FBuildInfoHelper::Connect()
+{
+	if (!bBuildListRetrieverInitComplete)
+	{
+		auto OnReadyForQuery = [ThisRef = AsShared()]()
+		{
+			ThisRef->Refresh(false);
+		};
+		BuildListRetriever->SetOnReadyForQuery(OnReadyForQuery);
+		bBuildListRetrieverInitComplete = true;
+	}
+
+	BuildListRetriever->ConnectToBuildService();
+}
+
+
+
+bool FBuildInfoHelper::IsConnected() const
+{
+	return BuildListRetriever->IsConnected();
+}
+
+
+
+bool FBuildInfoHelper::IsConfigured() const
+{
+	return BackendRetriever->IsConfigured() && UGSBuildInfoRetriever->IsConfigured();
+}
+
+FText FBuildInfoHelper::GetErrorText() const
+{
+	if (!BackendRetriever->IsConfigured() && !UGSBuildInfoRetriever->IsConfigured())
+	{
+		return LOCTEXT("RetrieversUnconfigured", "Horde and Backend URLs are not set in the Engine.ini");
+	}
+	else if (!BackendRetriever->IsConfigured())
+	{
+		return LOCTEXT("BackendUnconfigured", "Backend URL is not configured in the Engine.ini");
+	}
+	else if (!UGSBuildInfoRetriever->IsConfigured())
+	{
+		return LOCTEXT("UGSUnconfigured", "UGS URL is not configured in the Engine.ini");
+	}
+	else
+	{
+		return LOCTEXT("NoError", "Configured");
+	}
+}
+
+bool FBuildInfoHelper::IsRefreshing() const
+{
+	return bIsRefreshing || !BuildListRetriever->IsReadyForQuery();
+}
+
+
+
+void FBuildInfoHelper::Refresh( bool bFullRefresh )
+{
+	if (bFullRefresh)
+	{
+		BuildInfos.Reset();
+		BuildListRetriever->Refresh();
+	}
+	else if (!bIsRefreshing)
+	{
+		bIsRefreshing = true;
+
+		RefreshBackend();
+		RefreshBuildGroups();
+	}
+}
+
+
+
+void FBuildInfoHelper::CheckRefreshComplete()
+{
+	if ( !(bBuildGroupsRequestComplete && !BackendRetriever->IsGettingBackends() && !UGSBuildInfoRetriever->IsGettingUGSBuildInfo()))
+	{
+		return;
+	}
+
+	RebuildBuildInfos();
+
+	ExecuteOnGameThread(UE_SOURCE_LOCATION, [ThisRef = AsShared()]
+	{
+		if (ThisRef->OnBuildsRefreshed.IsSet())
+		{
+			ThisRef->OnBuildsRefreshed();
+		}
+		ThisRef->bIsRefreshing = false;
+	});
+}
+
+
+
+void FBuildInfoHelper::SetFilter( const FFilter& InFilter )
+{
+	Filter = InFilter;
+	Refresh(false);
+}
+
+
+
+void FBuildInfoHelper::SetProjectName( const FString& InProjectName )
+{
+	if (ProjectName != InProjectName)
+	{
+		ProjectName = InProjectName;
+
+		if (IsConnected())
+		{
+			Refresh(false);
+		}
+	}
+}
+
+void FBuildInfoHelper::SetBuildType( const FString& InBuildType )
+{
+	if (BuildType != InBuildType)
+	{
+		BuildType = InBuildType;
+
+		if (IsConnected())
+		{
+			Refresh(false);
+		}
+	}
+}
+
+void FBuildInfoHelper::SetBuildsRefreshedHandler(TFunction<void()> InOnBuildsRefreshed)
+{
+	OnBuildsRefreshed = InOnBuildsRefreshed;
+}
+
+
+void FBuildInfoHelper::RefreshBuildGroups()
+{
+	FilteredBuildGroups.Reset();
+	bBuildGroupsRequestComplete = false;
+
+	auto OnComplete = [ThisRef = AsShared()]()
+	{
+		ThisRef->bBuildGroupsRequestComplete = true;
+		
+		ThisRef->RefreshUGSBuildInfo(); // UGS build info requires the list of CLs from the build detail
+		ThisRef->CheckRefreshComplete();
+	};
+
+	// prepare the build detail query
+	FString Branch = FEngineVersion::Current().GetBranch();
+	Branch = Branch.Replace(TEXT("//"), TEXT("")).Replace(TEXT("/"), TEXT("-")).ToLower();
+
+	// launch the build query
+	bool bQuerying = BuildListRetriever->QueryBuildGroups(ProjectName, BuildType, Branch, 
+		[ThisRef = AsShared(), OnComplete](const TArrayView<TSharedRef<UE::Zen::Build::FBuildListRetriever::FBuildGroup>>& InBuildGroups) mutable
+		{
+			ThisRef->CreateFilteredBuildGroups(InBuildGroups);
+			OnComplete();
+		}
+	);
+	if (!bQuerying)
+	{
+		OnComplete();
+	}
+}
+
+
+
+void FBuildInfoHelper::CreateFilteredBuildGroups( const TArrayView<TSharedRef<UE::Zen::Build::FBuildListRetriever::FBuildGroup>>& BuildGroups )
+{
+	FilteredBuildGroups.Reset();
+
+	for (TSharedRef<UE::Zen::Build::FBuildListRetriever::FBuildGroup> BuildGroup : BuildGroups)
+	{
+		if (FilteredBuildGroups.Num() >= Filter.MaxItems)
+		{
+			break;
+		}
+
+		if (!Filter.MaxAge.IsZero() && (FDateTime::Now() - BuildGroup->CreatedAt) > Filter.MaxAge)
+		{
+			continue;
+		}
+
+		FilteredBuildGroups.Add(BuildGroup);
+	}
+}
+
+
+
+void FBuildInfoHelper::RebuildBuildInfos()
+{
+	BuildInfos.Reset(FilteredBuildGroups.Num());
+	AllKnownNamedArtifacts.Reset();
+
+	for ( TSharedRef<UE::Zen::Build::FBuildListRetriever::FBuildGroup> BuildGroup : FilteredBuildGroups)
+	{
+		TSharedPtr<FBuildInfo> BuildInfo = MakeShared<FBuildInfo>();
+		BuildInfo->Group = BuildGroup;
+
+		const TArray<FString>* BackendsPtr = BackendRetriever->GetBuildToBackendsMap().Find(BuildGroup->Changelist);
+		if (BackendsPtr != nullptr)
+		{
+			BuildInfo->Backends = (*BackendsPtr);
+		}
+
+		const TSharedRef<FUGSBuildInfoRetriever::FUGSBuildInfo>* UGSBuildInfoPtr = UGSBuildInfoRetriever->GetBuildToUGSBuildInfoMap().Find(BuildGroup->Changelist);
+		if (UGSBuildInfoPtr != nullptr)
+		{
+			BuildInfo->UGSBuildInfo = (*UGSBuildInfoPtr);
+		}
+
+		for ( const auto& Pair : BuildGroup->NamedArtifacts)
+		{
+			const FString& ArtifactName = Pair.Key; // @fixme: these are only 'platform' strings by convention unfortunately
+			const UE::Zen::Build::FBuildServiceInstance::FBuildRecord& BuildRecord = Pair.Value;
+			
+			FString Platform = GetPlatformNameFromNamedArtifact(ArtifactName, &BuildRecord);
+
+			FName PlatformName(Platform);
+			BuildInfo->Platforms.Add(PlatformName);
+
+			if (!BuildInfo->PlatformToArtifacts.Contains(PlatformName))
+			{
+				BuildInfo->PlatformToArtifacts.Add(PlatformName);
+			}
+			BuildInfo->PlatformToArtifacts[PlatformName].Add(ArtifactName);
+
+			AllKnownNamedArtifacts.Add(ArtifactName);
+		}
+		BuildInfo->Platforms.Sort( []( const FName& A, const FName& B )
+		{
+			return A.ToString() < B.ToString();
+		});
+
+		BuildInfos.Add(BuildInfo);
+	}
+
+	BuildTypesMap = BuildListRetriever->GenerateBuildTypesMap();
+}
+
+
+
+void FBuildInfoHelper::RefreshUGSBuildInfo()
+{
+	TArray<int32> Changelists;
+	for (TSharedPtr<UE::Zen::Build::FBuildListRetriever::FBuildGroup> BuildGroup : FilteredBuildGroups)
+	{
+		Changelists.AddUnique(BuildGroup->Changelist);
+	}
+
+	auto OnComplete = [ThisRef = AsShared()]()
+	{
+		ThisRef->CheckRefreshComplete();
+	};
+	UGSBuildInfoRetriever->GetUGSBuildInfoAsync(ProjectName, Changelists, OnComplete);
+}
+
+
+
+void FBuildInfoHelper::RefreshBackend()
+{
+	auto OnComplete = [ThisRef = AsShared()]()
+	{
+		ThisRef->CheckRefreshComplete();
+	};
+
+	BackendRetriever->GetBackendsAsync(OnComplete);
+}
+
+
+
+TSharedRef<UE::Zen::Build::FBuildServiceInstance> FBuildInfoHelper::GetServiceInstance() const
+{
+	return BuildListRetriever->GetServiceInstance();
+}
+
+
+
+const TArray<TSharedPtr<FBuildInfoHelper::FBuildInfo, ESPMode::ThreadSafe>>& FBuildInfoHelper::GetBuildInfos() const
+{
+	return BuildInfos;
+}
+
+const TMap<FString,TSharedRef<UE::Zen::Build::FBuildListRetriever::FBuildType, ESPMode::ThreadSafe>>& FBuildInfoHelper::GetBuildTypesMap() const
+{
+	return BuildTypesMap;
+}
+
+const TSet<FString>& FBuildInfoHelper::GetAllKnownNamedArtifacts() const
+{
+	return AllKnownNamedArtifacts;
+}
+
+
+FString FBuildInfoHelper::GetPlatformNameFromNamedArtifact( const FString& ArtifactName, const UE::Zen::Build::FBuildServiceInstance::FBuildRecord* BuildRecord ) const
+{
+	// build record's platform and the artifact name is only related to the UnrealTargetPlatform by convention :-(
+
+	// use the metadata platorm name if we have it, otherwise the artifact name will have to do
+	FString PlatformSearchString = ArtifactName;
+	if (BuildRecord != nullptr)
+	{
+		if (FCbFieldView CookPlatformField = BuildRecord->Metadata.FindIgnoreCase("platform"); CookPlatformField.HasValue() && !CookPlatformField.HasError())
+		{
+			FString CookPlatform(CookPlatformField.AsString());
+			if (!CookPlatform.IsEmpty())
+			{
+				PlatformSearchString = CookPlatform;
+			}
+		}
+	}
+
+
+	auto IsPlatform = [this](const FString& Value)
+	{
+		return KnownPlatformNames.Contains(Value.ToLower());
+	};
+
+	auto TrySetPlatform = [IsPlatform](FString Value)
+	{
+		Value.RemoveFromEnd(TEXT("Client"));
+		Value.RemoveFromEnd(TEXT("Game"));
+		Value.RemoveFromEnd(TEXT("Server"));
+		Value.RemoveFromEnd(TEXT("Editor"));
+
+		if (IsPlatform(Value))
+		{
+			Value.ReplaceInline(TEXT("Win64"), TEXT("Windows"));
+			return Value;
+		}
+
+		return FString();
+	};
+
+
+	auto TrySetPlatformFromItems = [TrySetPlatform](const TArray<FString>& Items)
+	{
+		for (const FString& Item : Items)
+		{
+			FString Platform = TrySetPlatform(Item);
+			if (!Platform.IsEmpty())
+			{
+				return Platform;
+			}
+		}
+
+		return FString();
+	};
+
+
+
+	FString Platform;
+
+	Platform = TrySetPlatform(PlatformSearchString); // raw Platform
+	if (!Platform.IsEmpty())
+	{
+		return Platform;
+	}
+
+	TArray<FString> DashedItems;
+	PlatformSearchString.ParseIntoArray(DashedItems, TEXT("-"));
+	Platform = TrySetPlatformFromItems(DashedItems); // Platform-BuildType-TargetType etc
+	if (!Platform.IsEmpty())
+	{
+		return Platform;
+	}
+
+	TArray<FString> UnderscoredItems;
+	PlatformSearchString.ParseIntoArray(UnderscoredItems, TEXT("_"));
+	Platform = TrySetPlatformFromItems(UnderscoredItems); // Platform_Flavor etc
+	if (!Platform.IsEmpty())
+	{
+		return Platform;
+	}
+
+
+	// cannot infer platform - try some fallbacks
+
+	if (DashedItems.Num() > 1)
+	{
+		return DashedItems[0];
+	}
+
+	if (UnderscoredItems.Num() > 1)
+	{
+		return UnderscoredItems[0];
+	}
+
+	return PlatformSearchString;
+}
+
+
+
+#undef LOCTEXT_NAMESPACE
+
diff --git a/Engine/Plugins/Developer/ProjectLauncher/Source/CommonLaunchExtensions/Private/BuildSync/BuildSyncLaunchExtension.cpp b/Engine/Plugins/Developer/ProjectLauncher/Source/CommonLaunchExtensions/Private/BuildSync/BuildSyncLaunchExtension.cpp
new file mode 100644
--- /dev/null
+++ b/Engine/Plugins/Developer/ProjectLauncher/Source/CommonLaunchExtensions/Private/BuildSync/BuildSyncLaunchExtension.cpp
@@ -0,0 +1,817 @@
+// Copyright Epic Games, Inc. All Rights Reserved.
+
+#include "BuildSync/BuildSyncLaunchExtension.h"
+#include "Model/ProjectLauncherModel.h"
+#include "Misc/EngineVersion.h"
+#include "Misc/ConfigCacheIni.h"
+#include "Misc/App.h"
+#include "Misc/DataDrivenPlatformInfoRegistry.h"
+#include "PlatformInfo.h"
+#include "Dom/JsonObject.h"
+#include "Dom/JsonValue.h"
+#include "Serialization/JsonReader.h"
+#include "Serialization/JsonSerializer.h"
+#include "DesktopPlatformModule.h"
+#include "Modules/BuildVersion.h"
+#include "Interfaces/IProjectManager.h"
+#include "Widgets/Shared/SCustomLaunchCombo.h"
+#include "Widgets/SBoxPanel.h"
+#include "Widgets/Input/SHyperlink.h"
+#include "Widgets/Input/SSearchBox.h"
+#include "Widgets/Input/SCheckBox.h"
+#include "Widgets/Text/STextBlock.h"
+#include "Internationalization/Regex.h"
+#include "GenericPlatform/GenericPlatformHttp.h"
+#include "BuildSync/SBuildSyncBuildCombo.h"
+
+#define LOCTEXT_NAMESPACE "FBuildSyncLaunchExtensionInstance"
+
+
+FBuildSyncLaunchExtensionInstance::FBuildSyncLaunchExtensionInstance( FArgs& InArgs ) 
+	: ProjectLauncher::FCustomUATCommandLaunchExtensionInstance(InArgs) 
+	, Owner(StaticCastSharedRef<FBuildSyncLaunchExtension>(InArgs.Extension))
+{
+	// cache platform information
+	GetConfigString(EConfig::PerProfile, TEXT("NamedArtifacts")).ParseIntoArray(CachedNamedArtifacts, TEXT(";"));
+	CachedCookPlatforms = GetCookPlatforms();
+
+	// create build info helper
+	BuildInfo = MakeShared<FBuildInfoHelper>();
+	BuildInfo->SetProjectName(GetProfile()->GetProjectName());
+	BuildInfo->SetBuildType(GetConfigString(EConfig::PerProfile, TEXT("BuildType"), FBuildInfoHelper::DefaultBuildType)); // do we want to set the default build type based on the content scheme?
+
+	auto OnRefreshed = [this]()
+	{
+		OnBuildsRefreshed();
+	};
+	BuildInfo->SetBuildsRefreshedHandler(OnRefreshed);	
+}
+
+FBuildSyncLaunchExtensionInstance::~FBuildSyncLaunchExtensionInstance()
+{
+	BuildInfo->SetBuildsRefreshedHandler(nullptr);
+	BuildInfo.Reset();
+	BuildCombo.Reset();
+}
+
+
+void FBuildSyncLaunchExtensionInstance::CustomizeTree( ProjectLauncher::FLaunchProfileTreeData& ProfileTreeData )
+{
+	FLaunchExtensionInstance::CustomizeTree(ProfileTreeData);
+
+	using namespace UE::Zen::Build;
+
+	auto IsEnabled = []()
+	{
+		return !GIsEditor;
+	};
+
+	auto IsEnabledAndReady = [this]()
+	{
+		return BuildInfo->IsConnected() && !BuildInfo->IsRefreshing() && !GIsEditor;
+	};
+
+	AddDefaultHeading(ProfileTreeData)
+		.AddWidget( LOCTEXT("ConnectionLabel", "Connection"),
+			{
+				.IsVisible = [this]() { return GetBuildServiceConnectionState() != FBuildServiceInstance::EConnectionState::ConnectionSucceeded; },
+				.IsEnabled = IsEnabled,
+			},
+			CreateBuildStorageConnectionWidget()
+		)
+		.AddWidget( LOCTEXT("BuildTypeLabel", "Build Type"),
+			{
+				.IsEnabled = IsEnabledAndReady,
+			},
+			CreateBuildTypeSelectionWidget()
+		)
+		.AddWidget( LOCTEXT("BuildLabel", "Build"),
+			{
+				.IsEnabled = IsEnabledAndReady,
+				.Validation = ProjectLauncher::FValidation({TEXT("Validation_BuildSync_NoBuild")}),
+			},
+			SNew(SHorizontalBox)
+			+SHorizontalBox::Slot()
+			.FillWidth(1)
+			.Padding(1,1)
+			.VAlign(VAlign_Center)
+			[			
+				SAssignNew(BuildCombo, SBuildSyncBuildCombo)
+				.OptionsSource(&FilteredBuildInfos)
+				.OnSelectionChanged(this, &FBuildSyncLaunchExtensionInstance::SetSelectedBuild)
+				.SelectedItem(this, &FBuildSyncLaunchExtensionInstance::GetSelectedBuild)
+				.GetItemSuitability(this, &FBuildSyncLaunchExtensionInstance::GetBuildSuitability)
+				.FilterWidget(CreateBuildFilterWidget())
+			]
+
+			+SHorizontalBox::Slot()
+			.AutoWidth()
+			.Padding(4,1)
+			.VAlign(VAlign_Center)
+			[
+				SNew(SButton)
+				.ButtonStyle(FAppStyle::Get(), "NoBorder")
+				.ToolTipText(LOCTEXT("RefreshTip", "Refresh the list of available builds"))
+				.OnClicked_Lambda([this]()
+				{ 
+					SelectedBuild.Reset(); // clear the build before doing a deep refresh - this means we auto-select the most suitable build
+					BuildInfo->Refresh(true); 
+					return FReply::Handled(); 
+				})
+				.Content()
+				[
+					SNew(SBox)
+					.WidthOverride(12)
+					.HeightOverride(12)
+					[
+						SNew(SImage)
+						.Image(FAppStyle::GetBrush("Icons.Refresh"))
+						.ColorAndOpacity(FSlateColor::UseSubduedForeground())
+					]
+				]
+			]
+		)
+		.AddWidget(LOCTEXT("NamedArtifactLabel", "Build Platform"),
+			{
+				.IsEnabled = IsEnabledAndReady,
+				.Validation = ProjectLauncher::FValidation({TEXT("Validation_BuildSync_NoArtifacts"), TEXT("Validation_BuildSync_ArtifactNotFound")}),
+			},
+			CreateNamedArtifactSelectionWidget()
+		)
+	;
+}
+
+
+// inspired by SBuildLogin from StorageServerWidgets
+BEGIN_SLATE_FUNCTION_BUILD_OPTIMIZATION
+TSharedRef<SWidget> FBuildSyncLaunchExtensionInstance::CreateBuildStorageConnectionWidget()
+{
+	using namespace UE::Zen::Build;
+
+	auto GetConnectionStatusText = [this]()
+	{
+		switch (GetBuildServiceConnectionState())
+		{
+			case FBuildServiceInstance::EConnectionState::NotStarted:					return LOCTEXT("BuildConnect_StatusValueNotConnected", "Not Connected");
+			case FBuildServiceInstance::EConnectionState::ConnectionInProgress:			return LOCTEXT("BuildConnect_StatusValueConnecting", "Connecting...");
+			case FBuildServiceInstance::EConnectionState::ConnectionSucceeded:			return FText::Format(LOCTEXT("BuildConnect_StatusValueConnected", "Connected to {0}"), FText::FromStringView(BuildInfo->GetServiceInstance()->GetEffectiveDomain()));
+			case FBuildServiceInstance::EConnectionState::ConnectionFailed:				return LOCTEXT("BuildConnect_StatusValueConnectionFailed", "Connection failed");
+			default:																	return LOCTEXT("BuildConnect_StatusValueError", "Error");
+		}
+	};
+
+	auto GetConnectLabel = [this]()
+	{
+		return (GetBuildServiceConnectionState() == FBuildServiceInstance::EConnectionState::NotStarted) ? LOCTEXT("BuildLogin_ConnectLink", "Click to Connect") : LOCTEXT("BuildLogin_ReconnectLink", "Reconnect");
+	};
+
+	return SNew(SHorizontalBox)
+
+		+SHorizontalBox::Slot()
+		.AutoWidth()
+		.Padding(2,4)
+		[
+			SNew(STextBlock)
+			.Text_Lambda(GetConnectionStatusText)
+			.Visibility_Lambda( [this]() { return (GetBuildServiceConnectionState() == FBuildServiceInstance::EConnectionState::NotStarted) ? EVisibility::Collapsed : EVisibility::Visible; } )
+			.Font(FCoreStyle::Get().GetFontStyle("SmallFont"))
+		]
+
+		+SHorizontalBox::Slot()
+		.AutoWidth()
+		.Padding(2,4)
+		[
+			SNew(SHyperlink)
+			.Text_Lambda(GetConnectLabel)
+			.OnNavigate_Lambda( [this]() { BuildInfo->Connect(); } )
+			.Visibility_Lambda([this]() { return (GetBuildServiceConnectionState() == FBuildServiceInstance::EConnectionState::NotStarted || GetBuildServiceConnectionState() == FBuildServiceInstance::EConnectionState::ConnectionFailed) ? EVisibility::Visible : EVisibility::Collapsed; } )
+		]
+	;
+}
+END_SLATE_FUNCTION_BUILD_OPTIMIZATION
+
+BEGIN_SLATE_FUNCTION_BUILD_OPTIMIZATION
+TSharedRef<SWidget> FBuildSyncLaunchExtensionInstance::CreateBuildFilterWidget()
+{
+	return SNew(SHorizontalBox)
+
+		// search filter box
+		+SHorizontalBox::Slot()
+		.FillWidth(1)
+		.Padding(2,1)
+		[
+			SNew(SSearchBox)
+			.HintText(LOCTEXT("BuildSearchHint", "Find build..."))
+			.InitialText_Lambda( [this]()
+			{
+				return FText::FromString(*FilterText); 
+			})
+			.OnTextChanged_Lambda( [this](const FText& SearchText)
+			{
+				FilterText = SearchText.ToString();
+				RefreshFilteredBuilds();
+			})
+		]
+
+		// show preflights
+		+SHorizontalBox::Slot()
+			.AutoWidth()
+			.Padding(2,1)
+			[
+				SNew(SCheckBox)
+				.ToolTipText(LOCTEXT("PreflightsToolTip", "Also show preflight builds"))
+				.IsChecked_Lambda( [this]() { return bShowPreflights ? ECheckBoxState::Checked : ECheckBoxState::Unchecked; } )
+				.OnCheckStateChanged_Lambda( [this](ECheckBoxState State) { bShowPreflights = (State == ECheckBoxState::Checked); RefreshFilteredBuilds(); } )
+				.IsFocusable(true)
+				[
+					SNew(STextBlock)
+					.Text(LOCTEXT("PreflightsLabel", "Preflights"))
+					.Font(FCoreStyle::Get().GetFontStyle("SmallFont"))
+				]
+			]
+
+		// show unsuitable builds
+		+SHorizontalBox::Slot()
+			.AutoWidth()
+			.Padding(2,1)
+			[
+				SNew(SCheckBox)
+				.ToolTipText(LOCTEXT("UnsuitableToolTip", "Also show builds that may not be suitable"))
+				.IsChecked_Lambda( [this]() { return bShowUnsuitable ? ECheckBoxState::Checked : ECheckBoxState::Unchecked; } )
+				.OnCheckStateChanged_Lambda( [this](ECheckBoxState State) { bShowUnsuitable = (State == ECheckBoxState::Checked); RefreshFilteredBuilds(); } )
+
+				[
+					SNew(STextBlock)
+					.Text(LOCTEXT("UnsuitableLabel", "Unsuitable"))
+					.Font(FCoreStyle::Get().GetFontStyle("SmallFont"))
+				]
+			]
+	;
+
+}
+END_SLATE_FUNCTION_BUILD_OPTIMIZATION
+
+BEGIN_SLATE_FUNCTION_BUILD_OPTIMIZATION
+TSharedRef<SWidget> FBuildSyncLaunchExtensionInstance::CreateBuildTypeSelectionWidget()
+{
+	using namespace UE::Zen::Build;
+
+	auto GetAvailableBuildTypes = [this]()
+	{
+		TArray<FString> BuildTypes;
+		BuildInfo->GetBuildTypesMap().GetKeys(BuildTypes);
+		return BuildTypes;
+	};
+
+	auto GetDisplayName = [this](const FString& BuildTypeName)
+	{
+		const TSharedRef<FBuildListRetriever::FBuildType, ESPMode::ThreadSafe>* BuildTypePtr = BuildInfo->GetBuildTypesMap().Find(BuildTypeName);
+		if (BuildTypePtr != nullptr)
+		{
+			return (*BuildTypePtr)->DisplayName;
+		}
+
+		return LOCTEXT("UnknownBuildTypeLabel", "Unknown");
+	};
+
+
+	return SNew(SCustomLaunchCombo<FString>)
+		.OnSelectionChanged(this, &FBuildSyncLaunchExtensionInstance::SetBuildType)
+		.SelectedItem(this, &FBuildSyncLaunchExtensionInstance::GetBuildType)
+		.Items_Lambda( [GetAvailableBuildTypes]() { return GetAvailableBuildTypes(); } )
+		.GetDisplayName_Lambda(GetDisplayName)
+	;
+}
+END_SLATE_FUNCTION_BUILD_OPTIMIZATION
+
+TArray<FString> FBuildSyncLaunchExtensionInstance::GetAvailableArtifactsForSelectedBuild() const
+{
+	TArray<FString> AvailableBuilds;
+	if (SelectedBuild.IsValid())
+	{
+		SelectedBuild->Group->NamedArtifacts.GenerateKeyArray(AvailableBuilds);
+		AvailableBuilds.Sort();
+	}
+	return AvailableBuilds;
+};
+
+TSharedRef<SWidget> FBuildSyncLaunchExtensionInstance::CreateNamedArtifactSelectionWidget()
+{
+	auto GetDisplayName = [this](const FString& Value)
+	{
+		FText Result = FText::FromString(Value);
+		if (SelectedBuild.IsValid() && SelectedBuild->Group->NamedArtifacts.Contains(Value))
+		{
+			return Result;
+		}
+		else
+		{
+			return FText::Format(LOCTEXT("NamedArtifactMissingLabel","{0} (not in this build)"), Result);
+		}
+	};
+
+	auto GetItemIcon = [this]( const FString& Item )
+	{
+		const UE::Zen::Build::FBuildServiceInstance::FBuildRecord* BuildRecord = SelectedBuild.IsValid() ? SelectedBuild->Group->NamedArtifacts.Find(Item) : nullptr;
+
+		FString Platform = BuildInfo->GetPlatformNameFromNamedArtifact(Item, BuildRecord);
+
+		FName IconName = FDataDrivenPlatformInfoRegistry::GetPlatformInfo(Platform).GetIconStyleName(EPlatformIconSize::Normal);
+		if (IconName.IsNone())
+		{
+			return FSlateIcon(FAppStyle::GetAppStyleSetName(), "Icons.FilledCircle");
+		}
+		else
+		{
+
+			return FSlateIcon(FAppStyle::GetAppStyleSetName(), IconName);
+		}
+	};
+
+	return SNew(SCustomLaunchStringCombo)
+		.OnSelectedItemsChanged(this, &FBuildSyncLaunchExtensionInstance::SetNamedArtifacts)
+		.SelectedItems(this, &FBuildSyncLaunchExtensionInstance::GetNamedArtifacts)
+		.Items(this, &FBuildSyncLaunchExtensionInstance::GetAvailableArtifactsForSelectedBuild)
+		.GetItemIcon_Lambda(GetItemIcon)
+		.GetItemIcon_Lambda(GetItemIcon)
+		.GetDisplayName_Lambda(GetDisplayName)
+		.ActionType(EUserInterfaceActionType::ToggleButton)
+	;
+}
+
+
+
+UE::Zen::Build::FBuildServiceInstance::EConnectionState FBuildSyncLaunchExtensionInstance::GetBuildServiceConnectionState() const
+{
+	return BuildInfo->GetServiceInstance()->GetConnectionState();
+}
+
+
+
+void FBuildSyncLaunchExtensionInstance::OnBuildsRefreshed()
+{
+	CacheUnsuitableBuilds();
+
+	const TSharedPtr<FBuildInfoHelper::FBuildInfo> OldSelectedBuild = SelectedBuild;
+
+	if (BuildInfo->GetBuildInfos().Num() > 0)
+	{
+		// we already have a selection & the list was refreshed - try to find the same build in the new list. if we fail, clear the selection
+		if (SelectedBuild != nullptr)
+		{
+			int32 Index = BuildInfo->GetBuildInfos().IndexOfByPredicate( [this]( const TSharedPtr<FBuildInfoHelper::FBuildInfo>& Build )
+			{
+				return (Build->Group->DisplayName == SelectedBuild->Group->DisplayName);
+			});
+			SelectedBuild = (Index == -1) ? nullptr : BuildInfo->GetBuildInfos()[Index];
+		}
+
+		if (SelectedBuild == nullptr)
+		{
+			AutoSelectBestBuild();
+		}
+	}
+	else
+	{
+		SelectedBuild = nullptr;
+	}
+
+	BroadcastEvent(TEXT("SelectedBuildChanged"));
+
+	RefreshFilteredBuilds();
+
+	if (CachedNamedArtifacts.IsEmpty())
+	{
+		AutoSelectNamedArtifacts();
+	}
+	else
+	{
+		// Remove artifacts that are not available for the selected build from the cached artifacts
+		TArray<FString> AvailableArtifacts = GetAvailableArtifactsForSelectedBuild();
+		CachedNamedArtifacts.RemoveAll([&AvailableArtifacts](const FString& Item)
+		{
+			return !AvailableArtifacts.Contains(Item);
+		});
+
+		SetNamedArtifacts(CachedNamedArtifacts);
+	}
+
+	GetProfile()->RefreshCustomWarningsAndErrors();
+}
+
+void FBuildSyncLaunchExtensionInstance::RefreshFilteredBuilds()
+{
+	FilteredBuildInfos.Reset();
+
+	for (TSharedPtr<FBuildInfoHelper::FBuildInfo> Build : BuildInfo->GetBuildInfos())
+	{		
+		// optionally filter out preflights
+		if (!bShowPreflights && Build->Group->bIsPreflight)
+		{
+			continue;
+		}
+
+		// optionally filter out unsuitable builds 
+		if (!bShowUnsuitable && CachedUnsuitableBuilds.Contains(Build))
+		{
+			continue;
+		}
+
+		// add search filtering
+		if (!FilterText.IsEmpty() && !Build->Group->DisplayName.Contains(FilterText))
+		{
+			continue;
+		}
+
+		FilteredBuildInfos.Add(Build);
+	}
+
+	if (BuildCombo.IsValid())
+	{
+		BuildCombo->RefreshOptions();
+	}
+}
+
+
+
+void FBuildSyncLaunchExtensionInstance::SetSelectedBuild(TSharedPtr<FBuildInfoHelper::FBuildInfo> Build)
+{
+	if (SelectedBuild != Build)
+	{
+		SelectedBuild = Build;
+		GetProfile()->RefreshCustomWarningsAndErrors();
+		CacheUnsuitableBuilds();
+		BroadcastEvent(TEXT("SelectedBuildChanged"));
+	}
+}
+
+TSharedPtr<FBuildInfoHelper::FBuildInfo> FBuildSyncLaunchExtensionInstance::GetSelectedBuild() const
+{
+	return SelectedBuild;
+}
+
+bool FBuildSyncLaunchExtensionInstance::GetBuildSuitability( TSharedPtr<FBuildInfoHelper::FBuildInfo> Build, FText* OutReason ) const
+{
+	if (OutReason == nullptr)
+	{
+		return !CachedUnsuitableBuilds.Contains(Build);
+	}
+	else
+	{
+		const FText* ReasonPtr = CachedUnsuitableBuilds.Find(Build);
+		if (ReasonPtr)
+		{
+			// build is unsuitable - return why
+			(*OutReason) = (*ReasonPtr);
+			return false;
+		}
+
+		// build is suitable
+		return true;
+	}
+}
+
+void FBuildSyncLaunchExtensionInstance::CacheUnsuitableBuilds()
+{
+	CachedUnsuitableBuilds.Reset();
+
+	for (TSharedPtr<FBuildInfoHelper::FBuildInfo> Build : BuildInfo->GetBuildInfos())
+	{
+		if (!CachedNamedArtifacts.IsEmpty())
+		{
+			TArray<FText> MissingNamedArtifacts;
+			for (const FString& NamedArtifact : CachedNamedArtifacts)
+			{
+				if (!Build->Group->NamedArtifacts.Contains(NamedArtifact))
+				{
+					MissingNamedArtifacts.Add(FText::FromString(*NamedArtifact) );
+				}
+			}
+
+			if (!MissingNamedArtifacts.IsEmpty())
+			{
+				FText Separator = LOCTEXT("Separator",", ");
+				FText Reason = FText::Format(LOCTEXT("BuildUnsuitable_MissingPlatform", "this build does not contain the selected Build Platform {0}"), FText::Join(Separator, MissingNamedArtifacts) );
+				CachedUnsuitableBuilds.Add(Build, Reason);
+				continue;
+			}
+		}
+
+
+	}
+}
+
+
+void FBuildSyncLaunchExtensionInstance::SetBuildType( FString InBuildType )
+{
+	SetNamedArtifacts(TArray<FString>());
+
+	SetConfigString(EConfig::PerProfile, TEXT("BuildType"), InBuildType );
+	BuildInfo->SetBuildType(InBuildType);
+}
+
+FString FBuildSyncLaunchExtensionInstance::GetBuildType() const
+{
+	FString Result = GetConfigString(EConfig::PerProfile, TEXT("BuildType"));
+	if (!Result.IsEmpty())
+	{
+		return Result;
+	}
+
+	if (auto* Pair = BuildInfo->GetBuildTypesMap().FindArbitraryElement())
+	{
+		return Pair->Key;
+	}
+	
+	return FBuildInfoHelper::DefaultBuildType;
+}
+void FBuildSyncLaunchExtensionInstance::SetNamedArtifacts( TArray<FString> InNamedArtifacts )
+{
+	CachedNamedArtifacts = InNamedArtifacts;
+	GetProfile()->RefreshCustomWarningsAndErrors();
+	CacheUnsuitableBuilds();
+
+	SetConfigString(EConfig::PerProfile, TEXT("NamedArtifacts"), FString::Join(CachedNamedArtifacts,TEXT(";")));
+}
+
+TArray<FString> FBuildSyncLaunchExtensionInstance::GetNamedArtifacts() const
+{
+	return CachedNamedArtifacts;
+}
+
+
+
+
+
+void FBuildSyncLaunchExtensionInstance::AutoSelectBestBuild()
+{
+	const TArray<const PlatformInfo::FTargetPlatformInfo*> PlatformInfos = GetModel()->GetPlatformInfosForAllCommands(GetProfile());
+
+	auto HasPlatforms = [PlatformInfos]( const TSharedPtr<FBuildInfoHelper::FBuildInfo>& Build, bool bNeedAllPlatforms )
+	{
+		for (const PlatformInfo::FTargetPlatformInfo* PlatformInfo : PlatformInfos)
+		{
+			const TArray<FString>* ArtifactsPtr = Build->PlatformToArtifacts.Find(PlatformInfo->IniPlatformName);
+			if (ArtifactsPtr == nullptr)
+			{
+				return false;
+			}
+			if (!bNeedAllPlatforms)
+			{
+				return true;
+			}
+		}
+
+		return true;
+	};
+
+	for ( int Pass = 0; Pass < 4; Pass++ ) // 0 = !PF + All Platforms   1 = !PF + Any Platform   2 = !PF    3 = anything
+	{
+		const bool bNeedAllPlatforms = (Pass == 0);
+
+		for (const TSharedPtr<FBuildInfoHelper::FBuildInfo>& Build : BuildInfo->GetBuildInfos())
+		{
+			if (Pass <= 2 && Build->Group->bIsPreflight)
+			{
+				continue;
+			}
+
+			if (Pass <= 1 && !HasPlatforms(Build, bNeedAllPlatforms))
+			{
+				continue;
+			}
+
+			SelectedBuild = Build;
+			break;
+		}
+
+		// stop when we find a build
+		if (SelectedBuild.IsValid())
+		{
+			// if we had to loosen the requirements, update the filter to match
+			bool bFiltersUpdated = false;
+			if (!bShowUnsuitable && Pass > 0)
+			{
+				bShowUnsuitable = true;
+				bFiltersUpdated = true;
+			}
+			if (!bShowPreflights && Pass > 2)
+			{
+				bShowPreflights = true;
+				bFiltersUpdated = true;
+			}
+
+			if (bFiltersUpdated)
+			{
+				RefreshFilteredBuilds();
+			}
+
+			break;
+		}
+	}
+}
+
+void FBuildSyncLaunchExtensionInstance::AutoSelectNamedArtifacts()
+{
+	if (SelectedBuild.IsValid())
+	{
+		TArray<FString> NewNamedArtifacts;
+
+		const TArray<const PlatformInfo::FTargetPlatformInfo*> PlatformInfos = GetModel()->GetPlatformInfosForAllCommands(GetProfile());
+		for (const PlatformInfo::FTargetPlatformInfo* PlatformInfo : PlatformInfos)
+		{
+			const TArray<FString>* ArtifactsPtr = SelectedBuild->PlatformToArtifacts.Find(PlatformInfo->IniPlatformName);
+			if (ArtifactsPtr != nullptr)
+			{
+				// sort artifacts by length... not an ideal way to find the most likely candidate. expectation is that less common candidates may have extra suffixes..
+				TArray<FString> Artifacts = (*ArtifactsPtr);
+				Artifacts.Sort( [](const FString& A, const FString& B)
+				{
+					return A.Len() < B.Len();
+				});
+
+				for (const FString& Artifact : Artifacts)
+				{
+					const UE::Zen::Build::FBuildServiceInstance::FBuildRecord* BuildRecord = SelectedBuild->Group->NamedArtifacts.Find(Artifact);
+					if (BuildRecord)
+					{
+						UE::Zen::Build::FBuildListRetriever::FBucketInfo BucketInfo;
+						if (UE::Zen::Build::FBuildListRetriever::GetBucketInfo(BuildRecord->BucketId, BucketInfo))
+						{
+							if (BucketInfo.BuildType == GetBuildType()) // would be nice to be able to pick the one that matches the profile's Build Target
+							{
+								NewNamedArtifacts.AddUnique(Artifact);
+								break; // only pick one artifact per platform
+							}
+						}
+					}
+				}
+			}
+		}
+
+		if (!NewNamedArtifacts.IsEmpty())
+		{
+			SetNamedArtifacts(NewNamedArtifacts);
+		}
+	}
+}
+
+TSet<const PlatformInfo::FTargetPlatformInfo*> FBuildSyncLaunchExtensionInstance::GetCookPlatforms() const
+{
+	TArray<const PlatformInfo::FTargetPlatformInfo*> PlatformInfos = GetModel()->GetPlatformInfosForAllCommands(GetProfile());
+
+	TSet<const PlatformInfo::FTargetPlatformInfo*> Result(PlatformInfos);
+	return MoveTemp(Result);
+}
+
+void FBuildSyncLaunchExtensionInstance::OnProjectChanged()
+{
+	CachedUnsuitableBuilds.Reset();
+	BuildInfo->SetProjectName(GetProfile()->GetProjectName());
+}
+
+void FBuildSyncLaunchExtensionInstance::OnPropertyChanged()
+{
+	
+	// see if the selected platforms have changed
+	TSet<const PlatformInfo::FTargetPlatformInfo*> CookPlatforms = GetCookPlatforms();
+	bool bPlatformsChanged = CookPlatforms.Num() != CachedCookPlatforms.Num();
+	if (!bPlatformsChanged)
+	{
+		for (const PlatformInfo::FTargetPlatformInfo* Platform : CookPlatforms) // why is there no Compare function for TSet?!
+		{
+			if (!CachedCookPlatforms.Contains(Platform))
+			{
+				bPlatformsChanged = true;
+				break;
+			}
+		}
+	}
+	if (bPlatformsChanged)
+	{
+		CachedCookPlatforms = CookPlatforms;
+		AutoSelectNamedArtifacts();
+	}
+}
+
+
+
+void FBuildSyncLaunchExtensionInstance::OnValidateProfile()
+{
+	if (GIsEditor)
+	{
+		GetProfile()->AddCustomError( TEXT("Validation_BuildSync_NotInEditor"), LOCTEXT("Validation_BuildSync_NotInEditor", "The Build Sync extension is not available in the editor because it may update the editor binaries. Use UnrealFrontend instead."));
+	}
+
+	if (!BuildInfo->IsConfigured())
+	{
+		GetProfile()->AddCustomWarning(TEXT("Validation_BuildSync_UnconfiguredBuildInfo"), BuildInfo->GetErrorText());
+	}
+
+	if (BuildInfo->IsRefreshing())
+	{
+		// fixme: disable the launch button but don't show error banner? it just seems distracting during the refresh for the banner to come in and out
+	}
+	else
+	{
+		if (!SelectedBuild.IsValid())
+		{
+			GetProfile()->AddCustomError( TEXT("Validation_BuildSync_NoBuild"), LOCTEXT("Validation_BuildSync_NoBuild", "No build selected"));
+		}
+		if (CachedNamedArtifacts.IsEmpty())
+		{
+			GetProfile()->AddCustomError( TEXT("Validation_BuildSync_NoArtifacts"), LOCTEXT("Validation_BuildSync_NoArtifacts", "No build platforms selected"));
+		}
+		if (!CachedNamedArtifacts.IsEmpty() && SelectedBuild.IsValid())
+		{
+			bool bHasMissingNamedArtifacts = CachedNamedArtifacts.ContainsByPredicate( [this]( const FString& NamedArtifact )
+			{
+				return !SelectedBuild->Group->NamedArtifacts.Contains(NamedArtifact);
+			});
+
+			if (bHasMissingNamedArtifacts)
+			{
+				GetProfile()->AddCustomError( TEXT("Validation_BuildSync_ArtifactNotFound"), LOCTEXT("Validation_BuildSync_ArtifactNotFound", "This build does not contain the selected build platforms"));
+			}
+		}
+	}
+}
+
+
+
+
+void FBuildSyncLaunchExtensionInstance::CustomizeUATCommandLine( FString& InOutCommandLine )
+{
+	if (!SelectedBuild.IsValid())
+	{
+		return;
+	}
+
+	FString CommandLine;
+
+	CommandLine += FString::Printf(TEXT(" -project=\"%s\" "), *FPaths::ConvertRelativePathToFull(GetProfile()->GetProjectPath()));
+	CommandLine += FString::Printf(TEXT(" -CL=%s"), *SelectedBuild->Group->CommitIdentifier);
+
+	FString Branch = FEngineVersion::Current().GetBranch();
+	Branch = Branch.Replace(TEXT("//"), TEXT("")).Replace(TEXT("/"), TEXT("-")).ToLower();
+	CommandLine += FString::Printf(TEXT(" -Stream=%s"), *Branch);
+
+	CommandLine += TEXT(" -getbuild");
+	CommandLine += FString::Printf(TEXT(" -buildtype=%s"), *GetBuildType());
+	CommandLine += FString::Printf(TEXT(" -buildplatform=%s"), *FString::Join(CachedNamedArtifacts, TEXT("+")));
+
+	InOutCommandLine += CommandLine;
+}
+
+void FBuildSyncLaunchExtensionInstance::InternalInitialize()
+{
+	FLaunchExtensionInstance::InternalInitialize();
+
+	// now the command has been added, start the connection process if we have a project
+	if (!GIsEditor && !GetProfile()->GetProjectName().IsEmpty())
+	{
+		BuildInfo->Connect();
+	}
+}
+
+void FBuildSyncLaunchExtensionInstance::OnUATCommandAdded( ILauncherProfileUATCommandRef InUATCommand )
+{
+	InUATCommand->SetUATCommand(TEXT("BuildAcquire")); // RunUAT BuildAcquire [...] 
+}
+
+TArray<FString> FBuildSyncLaunchExtensionInstance::GetSelectedBuildBackends() const
+{
+	return SelectedBuild ? SelectedBuild->Backends : TArray<FString>();
+}
+
+
+
+
+
+
+
+
+
+
+FBuildSyncLaunchExtension::FBuildSyncLaunchExtension()
+{
+}
+
+TSharedPtr<ProjectLauncher::FLaunchExtensionInstance> FBuildSyncLaunchExtension::CreateInstanceForProfile( ProjectLauncher::FLaunchExtensionInstance::FArgs& InArgs )
+{
+	return MakeShared<FBuildSyncLaunchExtensionInstance>(InArgs);
+}
+
+const TCHAR* FBuildSyncLaunchExtension::GetInternalName() const
+{
+	return TEXT("BuildSync");
+}
+
+FText FBuildSyncLaunchExtension::GetDisplayName() const
+{
+	return LOCTEXT("ExtensionName", "Build Sync");
+}
+
+
+
+#undef LOCTEXT_NAMESPACE
diff --git a/Engine/Plugins/Developer/ProjectLauncher/Source/CommonLaunchExtensions/Private/BuildSync/SBuildSyncBuildCombo.cpp b/Engine/Plugins/Developer/ProjectLauncher/Source/CommonLaunchExtensions/Private/BuildSync/SBuildSyncBuildCombo.cpp
new file mode 100644
--- /dev/null
+++ b/Engine/Plugins/Developer/ProjectLauncher/Source/CommonLaunchExtensions/Private/BuildSync/SBuildSyncBuildCombo.cpp
@@ -0,0 +1,496 @@
+// Copyright Epic Games, Inc. All Rights Reserved.
+
+#include "BuildSync/SBuildSyncBuildCombo.h"
+
+#include "SlateOptMacros.h"
+#include "Widgets/SBoxPanel.h"
+#include "Widgets/Input/SComboBox.h"
+#include "Widgets/Input/SButton.h"
+#include "Widgets/Text/STextBlock.h"
+#include "Widgets/Images/SImage.h"
+#include "Widgets/Colors/SColorBlock.h"
+#include "Math/ColorList.h"
+#include "Misc/DataDrivenPlatformInfoRegistry.h"
+
+
+#define LOCTEXT_NAMESPACE "SBuildSyncBuildCombo"
+
+
+BEGIN_SLATE_FUNCTION_BUILD_OPTIMIZATION
+void SBuildSyncBuildCombo::Construct(const FArguments& InArgs)
+{
+	OptionsSource = InArgs._OptionsSource;
+	OnSelectionChanged = InArgs._OnSelectionChanged;
+	GetItemSuitability = InArgs._GetItemSuitability;
+
+	SelectedItem = InArgs._SelectedItem;
+	RequiredPlatforms = InArgs._RequiredPlatforms;
+
+	check (InArgs._FilterWidget.IsValid());
+
+	ChildSlot
+	[
+		SAssignNew(ComboButton, SComboButton)
+		.IsFocusable(true)
+		.CollapseMenuOnParentFocus(true)
+		.ButtonContent()
+		[
+			SNew(STextBlock)
+			.Text(this, &SBuildSyncBuildCombo::GetSelectedBuildInfoName)
+			.Font(FCoreStyle::Get().GetFontStyle("SmallFont"))
+		]
+		.MenuContent()
+		[
+			SNew(SBox)
+			.MaxDesiredHeight(512)
+			[
+				SNew(SVerticalBox)
+				+SVerticalBox::Slot()
+				.AutoHeight()
+				.Padding(8)
+				[
+					InArgs._FilterWidget.ToSharedRef()
+				]
+
+				+SVerticalBox::Slot()
+				.FillHeight(1)
+				.Padding(8)
+				[
+					SAssignNew(BuildInfosListView, SListView<TSharedPtr<FBuildInfoHelper::FBuildInfo>>)
+					.ListItemsSource(OptionsSource)
+					.OnGenerateRow( this, &SBuildSyncBuildCombo::OnGenerateBuildInfoRow )
+					.OnSelectionChanged( this, &SBuildSyncBuildCombo::OnBuildInfoSelectionChanged )
+					.SelectionMode(ESelectionMode::SingleToggle)
+					.IsFocusable(true)
+					.OnKeyDownHandler_Lambda( [this](const FGeometry& MyGeometry, const FKeyEvent& InKeyEvent)
+					{
+						if (InKeyEvent.GetKey() == EKeys::Enter && BuildInfosListView->GetNumItemsSelected() > 0)
+						{
+							OnBuildInfoSelectionChanged( BuildInfosListView->GetSelectedItems()[0], ESelectInfo::OnKeyPress);
+							return FReply::Handled();
+						}
+
+						return FReply::Unhandled();
+					})
+				]
+			]
+		]
+		.OnMenuOpenChanged_Lambda( [this](bool bOpen)
+		{
+			if (bOpen)
+			{
+				if (SelectedItem.Get().IsValid())
+				{
+					BuildInfosListView->RequestScrollIntoView(SelectedItem.Get());
+					FSlateApplication::Get().SetKeyboardFocus(BuildInfosListView, EFocusCause::SetDirectly);
+				}
+			}
+			else
+			{
+				FSlateApplication::Get().SetKeyboardFocus(ComboButton, EFocusCause::SetDirectly);
+			}
+		})
+	];
+
+	ComboButton->SetMenuContentWidgetToFocus(BuildInfosListView);
+
+}
+END_SLATE_FUNCTION_BUILD_OPTIMIZATION
+
+
+BEGIN_SLATE_FUNCTION_BUILD_OPTIMIZATION
+TSharedRef<ITableRow> SBuildSyncBuildCombo::OnGenerateBuildInfoRow( TSharedPtr<FBuildInfoHelper::FBuildInfo> BuildInfo, const TSharedRef<STableViewBase>& OwnerTable) const
+{
+	if (!BuildInfo.IsValid())
+	{
+		return SNew(STableRow<TSharedPtr<FBuildInfoHelper::FBuildInfo>>, OwnerTable)
+		[
+			SNullWidget::NullWidget
+		];
+	}
+	
+	static const FColor BadgeColors[] =
+	{
+		FColorList::LightBlue,			// Starting
+		FColorList::OrangeRed,			// Failure
+		FColorList::Orange,				// Warning
+		FColorList::MediumSpringGreen,	// Success
+		FColorList::LightGrey,			// Skipped
+		FColorList::LightGrey,			// Unknown
+
+	};
+	static_assert(UE_ARRAY_COUNT(BadgeColors) == (int)FUGSBuildInfoRetriever::FUGSBadge::EState::MAX);
+
+
+	auto GetSuitablilityColor = [this, BuildInfo]()
+	{
+		bool bIsSuitable = !GetItemSuitability.IsBound() || GetItemSuitability.Execute(BuildInfo, nullptr);
+		return bIsSuitable ? FSlateColor::UseForeground() : FSlateColor::UseSubduedForeground();
+	};
+
+	auto GetSuitabilityToolTip = [this, BuildInfo]()
+	{
+		FText Text = FText::GetEmpty();
+		if (GetItemSuitability.IsBound())
+		{
+			GetItemSuitability.Execute(BuildInfo, &Text);
+		}
+
+		return Text;
+	};
+
+
+	// extract the build suffix - anything after ++Foo+Bar-CL-123
+	FString BuildVersionSuffix;
+	FString Suffix = BuildInfo->Group->Suffix;
+	if (Suffix.IsEmpty())
+	{
+		int32 Idx = BuildInfo->Group->DisplayName.Find(BuildInfo->Group->CommitIdentifier);
+		if (Idx != -1)
+		{
+			Suffix = BuildInfo->Group->DisplayName.RightChop(Idx + BuildInfo->Group->CommitIdentifier.Len() + 1);
+		}
+	}
+	if (!Suffix.IsEmpty())
+	{
+		BuildVersionSuffix = TEXT("[") + Suffix + TEXT("]");
+	}
+
+	// prepare the first row
+	TSharedPtr<SHorizontalBox> BuildInfoBox;
+	TSharedPtr<SVerticalBox> BuildInfoInfoBox;
+
+	TSharedRef<SBox> Result = 
+		SNew(SBox)
+		.WidthOverride(512) // maintain everything at the same width
+		.Padding(4,8)
+		[
+			SNew(SHorizontalBox)
+
+			// check mark for selected item
+			+SHorizontalBox::Slot()
+			.AutoWidth()
+			.VAlign(VAlign_Center)
+			.Padding(4)
+			[
+				SNew(SImage)
+				.Image(FAppStyle::Get().GetBrush("Icons.Check"))
+				.Visibility_Lambda( [this, BuildInfo]() { return (SelectedItem.Get() == BuildInfo) ? EVisibility::Visible : EVisibility::Hidden; } )
+			]
+
+			// item
+			+SHorizontalBox::Slot()
+			.FillWidth(1)
+			[
+				SAssignNew(BuildInfoInfoBox, SVerticalBox)
+				+SVerticalBox::Slot()
+				.AutoHeight()
+				[
+					SAssignNew(BuildInfoBox, SHorizontalBox)
+
+					+SHorizontalBox::Slot()
+					.AutoWidth()
+
+					// changelist number
+					+SHorizontalBox::Slot()
+					.AutoWidth()
+					.Padding(1)
+					.VAlign(VAlign_Center)
+					[
+						SNew(STextBlock)
+						.Text(FText::FromString(*BuildInfo->Group->CommitIdentifier))
+						.Font(FCoreStyle::Get().GetFontStyle("NormalFontBold"))
+						.ColorAndOpacity_Lambda(GetSuitablilityColor)
+					]
+
+					// additional detail (preflight code etc)
+					+SHorizontalBox::Slot()
+					.AutoWidth()
+					.Padding(8,1)
+					.VAlign(VAlign_Center)
+					[
+						SNew(STextBlock)
+						.Visibility(BuildVersionSuffix.IsEmpty() ? EVisibility::Collapsed : EVisibility::Visible)
+						.Text(FText::FromString(*BuildVersionSuffix))
+						.Font(FCoreStyle::Get().GetFontStyle("NormalFont"))
+						.ColorAndOpacity_Lambda(GetSuitablilityColor)
+					]
+
+					// age
+					+SHorizontalBox::Slot()
+					.AutoWidth()
+					.Padding(8,1)
+					.VAlign(VAlign_Center)
+					[
+						SNew(STextBlock)
+						.Text(GetAge(BuildInfo->Group->CreatedAt))
+						.Font(FCoreStyle::Get().GetFontStyle("NormalFontItalic"))
+						.ColorAndOpacity_Lambda(GetSuitablilityColor)
+					]
+				]
+			]
+		]
+	;
+
+
+	// platform icons & artifacts. use a plain circle for platforms we don't know about
+	for (const FName& Platform : BuildInfo->Platforms)
+	{
+		FName IconName = FDataDrivenPlatformInfoRegistry::GetPlatformInfo(Platform).GetIconStyleName(EPlatformIconSize::Normal);
+
+		const FSlateBrush* Brush = IconName.IsNone() ? FAppStyle::Get().GetBrush("Icons.FilledCircle") : FAppStyle::Get().GetBrush(IconName);
+		const FSlateColor& Color = IconName.IsNone() ? FSlateColor::UseSubduedForeground() : GetSuitablilityColor();
+
+		BuildInfoBox->AddSlot()
+			.AutoWidth()
+			.Padding(1)
+			.VAlign(VAlign_Center)
+			[
+				SNew(SImage)
+				.Image(Brush)
+				.ToolTipText(FText::FromString(FString::Join(BuildInfo->PlatformToArtifacts[Platform], TEXT("\n"))))
+				.ColorAndOpacity(Color)
+				.DesiredSizeOverride(FVector2D(12,12))
+			]
+		;
+	}
+
+
+
+
+	// add the UGS build information if we have it
+	if (BuildInfo->UGSBuildInfo.IsValid())
+	{
+		// success / failure badge (not using this at the moment - setting the user count text color instead)
+		//if (BuildInfo->UGSBuildInfo->NumSuccess > 0 || BuildInfo->UGSBuildInfo->NumFailed > 0)
+		//{
+		//	bool bSuccess = (BuildInfo->UGSBuildInfo->NumSuccess > BuildInfo->UGSBuildInfo->NumFailed);
+		//
+		//	BuildInfoBox->InsertSlot(0)
+		//		.AutoWidth()
+		//		.Padding(1)
+		//		[
+		//			SNew(SImage)
+		//			.Image(FAppStyle::Get().GetBrush("Icons.FilledCircle"))
+		//			.ColorAndOpacity( bSuccess ? FColorList::MediumSpringGreen : FColorList::OrangeRed )
+		//		]
+		//	;
+		//}
+
+
+
+
+		// user count
+		if (BuildInfo->UGSBuildInfo->NumUsers > 0)
+		{
+			FSlateColor Color = FSlateColor::UseForeground();
+			if (BuildInfo->UGSBuildInfo->NumSuccess > 0 || BuildInfo->UGSBuildInfo->NumFailed > 0)
+			{
+				bool bSuccess = (BuildInfo->UGSBuildInfo->NumSuccess > BuildInfo->UGSBuildInfo->NumFailed);
+				Color = bSuccess ? FColorList::MediumSpringGreen : FColorList::OrangeRed;
+			}
+
+			BuildInfoBox->AddSlot()
+				.FillWidth(1)
+				.Padding(8,1)
+				.VAlign(VAlign_Bottom)
+				.HAlign(HAlign_Right)
+				[
+					SNew(STextBlock)
+					.Text(FText::Format(LOCTEXT("NumUsersLabel", "{0} users"), BuildInfo->UGSBuildInfo->NumUsers ))
+					.Font(FCoreStyle::Get().GetFontStyle("SmallFont"))
+					.ColorAndOpacity(Color)
+				]
+			;
+		}
+	
+
+		// add UGS badges array if there's anything to display
+		if (BuildInfo->UGSBuildInfo->Badges.Num() > 0)
+		{
+			TSharedRef<SHorizontalBox> BadgesBox = SNew(SHorizontalBox);
+			BuildInfoInfoBox->AddSlot()
+				.Padding(8,1)
+				[
+					SNew(SBox)
+					.HeightOverride(14)
+					[
+						BadgesBox
+					]
+				]
+			;
+
+			for ( const FUGSBuildInfoRetriever::FUGSBadge& Badge : BuildInfo->UGSBuildInfo->Badges)
+			{
+				static bool bShowBadgesWithNoUrl = false; // todo: could make this public
+				if (!bShowBadgesWithNoUrl && Badge.URL.IsEmpty())
+				{
+					continue;
+				}
+
+				auto OnClick = [URL = Badge.URL](const FGeometry& InGeometry, const FPointerEvent& InEvent)
+				{
+					if (FSlateApplication::Get().GetModifierKeys().IsControlDown())
+					{
+						FPlatformProcess::LaunchURL(*URL, nullptr, nullptr); 
+						return FReply::Handled(); 
+					}
+
+					return FReply::Unhandled(); 
+				};
+
+				auto GetCursor = []()
+				{
+					return FSlateApplication::Get().GetModifierKeys().IsControlDown() ? EMouseCursor::Hand : EMouseCursor::Default;
+				};
+
+				auto GetColor = [this, State = Badge.State, BuildInfo]()
+				{
+					FLinearColor Color = BadgeColors[(int8)State];
+					if (GetItemSuitability.IsBound() && !GetItemSuitability.Execute(BuildInfo, nullptr))
+					{
+						Color = Color.Desaturate(0.25f);
+					}
+
+					return Color;
+				};
+
+
+				FText ToolTipText = FText::Format(LOCTEXT("UgsBadgeTip", "{0}\nCtrl+Click : {1}"), FText::FromString(*Badge.Name), FText::FromString(*Badge.URL) );
+				BadgesBox->AddSlot()
+					.FillWidth(1)
+					.Padding(1,0,0,0)
+					.VAlign(VAlign_Fill)
+					[
+						SNew(SBorder)
+						.BorderImage(FAppStyle::GetBrush("NoBorder"))
+						.Padding(FMargin(0,4))
+						.Cursor_Lambda(GetCursor)
+						.OnMouseButtonDown_Lambda(OnClick)
+						.ToolTipText(ToolTipText)
+						[
+							SNew(SColorBlock)
+							.Color_Lambda(GetColor)
+						]
+					]
+				;
+			}
+		}
+	}
+
+
+
+	// add backends display row, if there's anything to display
+	TSharedPtr<SHorizontalBox> BackendsBox;
+	if (BuildInfo->Backends.Num() > 0)
+	{
+		BackendsBox = SNew(SHorizontalBox);
+		BuildInfoInfoBox->AddSlot()
+			.AutoHeight()
+			.Padding(8,1)
+			[
+				BackendsBox.ToSharedRef()
+			]
+		;
+
+		for (const FString& Backend : BuildInfo->Backends)
+		{
+			BackendsBox->AddSlot()
+				.AutoWidth()
+				.Padding(4)
+				.VAlign(VAlign_Center)
+				.HAlign(HAlign_Right)
+				[
+					SNew(SBorder)
+					.Padding(2,4)
+					.BorderImage(FCoreStyle::Get().GetBrush("FloatingBorder"))
+					.BorderBackgroundColor(FLinearColor::Black)
+					[
+						SNew(STextBlock)
+						.Text(FText::FromString(*Backend))
+						.ColorAndOpacity(FLinearColor::White)
+						.Justification(ETextJustify::Center)
+						.Font(FCoreStyle::Get().GetFontStyle("SmallFont"))
+					]
+				]
+			;
+		}
+	}
+
+	return SNew(STableRow<TSharedPtr<FBuildInfoHelper::FBuildInfo>>, OwnerTable)
+	.Style(&FAppStyle::Get().GetWidgetStyle<FTableRowStyle>("ComboBox.Row"))
+    .ShowSelection(true)
+	.ToolTipText_Lambda(GetSuitabilityToolTip)
+	[
+		Result
+	];
+}
+END_SLATE_FUNCTION_BUILD_OPTIMIZATION
+
+
+
+
+
+
+FText SBuildSyncBuildCombo::GetAge( const FDateTime& DateTime ) const
+{
+	FTimespan Age = FDateTime::Now() - DateTime;
+	if (Age.GetDays() > 0)
+	{
+		return FText::Format(LOCTEXT("Timespan_DaysHoursMin", "({0}days {1}hrs {2}mins old)"), Age.GetDays(), Age.GetHours(), Age.GetMinutes() );
+	}
+	else if (Age.GetHours() > 0)
+	{
+		return FText::Format(LOCTEXT("Timespan_HoursMin", "({0}hrs {1}mins old)"), Age.GetHours(), Age.GetMinutes() );
+	}
+	else
+	{
+		return FText::Format(LOCTEXT("Timespan_Min", "({0}mins old)"), Age.GetMinutes() );
+	}
+}
+
+
+
+
+void SBuildSyncBuildCombo::OnBuildInfoSelectionChanged( TSharedPtr<FBuildInfoHelper::FBuildInfo> BuildInfo, ESelectInfo::Type InSelectInfo )
+{
+	if (InSelectInfo != ESelectInfo::Direct)
+	{
+		// using ESelecitonMode::SingleToggle so the combo box can close even if you select the same item - so need to ignore the 'unselect'
+		if (BuildInfo.IsValid()) 
+		{
+			OnSelectionChanged.ExecuteIfBound(BuildInfo);
+		}
+
+		if (InSelectInfo == ESelectInfo::OnKeyPress || InSelectInfo == ESelectInfo::OnMouseClick)
+		{
+			ComboButton->SetIsOpen( false );
+		}
+	}
+}
+
+
+FText SBuildSyncBuildCombo::GetSelectedBuildInfoName() const
+{
+	TSharedPtr<FBuildInfoHelper::FBuildInfo> SelectedBuildInfo = SelectedItem.Get();
+	return SelectedBuildInfo.IsValid() ? FText::FromString(*SelectedBuildInfo->Group->DisplayName) : LOCTEXT("NoBuild", "(no build)");
+}
+
+
+
+
+
+void SBuildSyncBuildCombo::RefreshOptions()
+{
+	BuildInfosListView->RequestListRefresh();
+
+	BuildInfosListView->ClearSelection();
+	if (SelectedItem.Get().IsValid())
+	{
+		BuildInfosListView->SetSelection(SelectedItem.Get());
+		BuildInfosListView->RequestScrollIntoView(SelectedItem.Get());
+	}
+}
+
+
+#undef LOCTEXT_NAMESPACE
diff --git a/Engine/Plugins/Developer/ProjectLauncher/Source/CommonLaunchExtensions/Private/CommonLaunchExtensionsModule.cpp b/Engine/Plugins/Developer/ProjectLauncher/Source/CommonLaunchExtensions/Private/CommonLaunchExtensionsModule.cpp
--- a/Engine/Plugins/Developer/ProjectLauncher/Source/CommonLaunchExtensions/Private/CommonLaunchExtensionsModule.cpp
+++ b/Engine/Plugins/Developer/ProjectLauncher/Source/CommonLaunchExtensions/Private/CommonLaunchExtensionsModule.cpp
@@ -9,6 +9,7 @@
 #include "AdvancedCook/AdvancedCookLaunchExtension.h"
 #include "ProfileWizard/ProfileWizardLaunchExtension.h"
 #include "DeprecatedProperties/DeprecatedPropertiesLaunchExtension.h"
+#include "BuildSync/BuildSyncLaunchExtension.h"
 
 class FCommonLaunchExtensionsModule : public IModuleInterface
 {
@@ -22,6 +23,7 @@ public:
 		Extensions.Add( MakeShared<FAdvancedCookLaunchExtension>() );
 		Extensions.Add( MakeShared<FProfileWizardLaunchExtension>() );
 		Extensions.Add( MakeShared<FDeprecatedPropertiesLaunchExtension>() );
+		Extensions.Add( MakeShared<FBuildSyncLaunchExtension>() );
 
 		for (TSharedRef<ProjectLauncher::FLaunchExtension> Extension : Extensions)
 		{
diff --git a/Engine/Plugins/Developer/ProjectLauncher/Source/CommonLaunchExtensions/Private/Shared/BackendRetriever.cpp b/Engine/Plugins/Developer/ProjectLauncher/Source/CommonLaunchExtensions/Private/Shared/BackendRetriever.cpp
new file mode 100644
--- /dev/null
+++ b/Engine/Plugins/Developer/ProjectLauncher/Source/CommonLaunchExtensions/Private/Shared/BackendRetriever.cpp
@@ -0,0 +1,164 @@
+// Copyright Epic Games, Inc. All Rights Reserved.
+
+#include "Shared/BackendRetriever.h"
+#include "HAL/PlatformMisc.h"
+#include "Misc/EngineVersion.h"
+#include "Misc/ConfigCacheIni.h"
+#include "Dom/JsonObject.h"
+#include "Dom/JsonValue.h"
+#include "Serialization/JsonReader.h"
+#include "Serialization/JsonSerializer.h"
+#include "PlatformHttp.h"
+#include "HttpModule.h"
+#include "Interfaces/IHttpRequest.h"
+#include "Interfaces/IHttpResponse.h"
+
+FBackendRetriever::FBackendRetriever()
+{
+	GConfig->GetString(TEXT("ProjectLauncher"), TEXT("BackendQueryUrl"), BackendQueryUrl, GEngineIni);
+}
+
+void FBackendRetriever::GetBackendsAsync(FOnBackendRetrievalComplete&& InOnComplete)
+{
+	if (BackendQueryUrl.IsEmpty())
+	{
+		return;
+	}
+
+	if (bIsGettingBackends)
+	{
+		if (InOnComplete)
+		{
+			InOnComplete();
+		}
+		return;
+	}
+
+	bIsGettingBackends = true;
+	OnComplete = MoveTemp(InOnComplete);
+	ErrorMessage = FString();
+	BuildToBackendsMap.Reset();
+
+
+	// build the backend query request
+	FString Stream = FEngineVersion::Current().GetBranch();
+	TSharedRef<IHttpRequest, ESPMode::ThreadSafe> HttpRequest = FHttpModule::Get().CreateRequest();
+	HttpRequest->SetVerb(TEXT("POST"));
+	HttpRequest->SetURL(BackendQueryUrl);
+	HttpRequest->SetHeader(TEXT("Content-Type"), TEXT("application/json"));
+	HttpRequest->SetContentAsString(FString::Printf(TEXT("{\"stream\":\"%s\"}"), *Stream));
+	HttpRequest->OnProcessRequestComplete().BindLambda( 
+		[this](FHttpRequestPtr Request, FHttpResponsePtr Response, bool bWasSuccessful)
+		{
+			if (bWasSuccessful && Response->GetResponseCode() == 200)
+			{
+				ParseBackendFromOutput(Response->GetContentAsString());
+			}
+			else if (Response)
+			{
+				ErrorMessage = FString::Printf(TEXT("Failed to query backend. error code %d: %s"), Response->GetResponseCode(), *Response->GetContentAsString());
+			}
+			else
+			{
+				ErrorMessage = TEXT("Failed to query backend. no response from server");
+			}
+			Completed();
+		}
+	);
+
+	// launch the request
+	if (!HttpRequest->ProcessRequest())
+	{
+		ErrorMessage = FString::Printf(TEXT("Failed to start backend query for stream %s"), *Stream);
+		Completed();
+	}
+}
+
+
+
+void FBackendRetriever::ParseBackendFromOutput( const FString& QueryResult)
+{
+	TSharedPtr<FJsonObject> ResponseJson;
+	if (FJsonSerializer::Deserialize(TJsonReaderFactory<>::Create(QueryResult), ResponseJson) && ResponseJson.IsValid())
+	{
+		const TArray<TSharedPtr<FJsonValue>>* AttributesJsonArray;
+		if (ResponseJson->TryGetArrayField(TEXT("list"), AttributesJsonArray))
+		{
+			for (int32 NameIdx = 0; NameIdx < AttributesJsonArray->Num(); ++NameIdx)
+			{
+				TSharedPtr<FJsonValue> BuildInfoJson = (*AttributesJsonArray)[NameIdx];
+
+				const TSharedPtr<FJsonObject>* BuildInfoObj;
+				if (BuildInfoJson->TryGetObject(BuildInfoObj))
+				{
+					FString EpicApp, BuildInfo;
+					if ((*BuildInfoObj)->TryGetStringField(TEXT("epic_app"), EpicApp) &&
+						(*BuildInfoObj)->TryGetStringField(TEXT("build_info"), BuildInfo))
+					{
+						// Build is in the format of "++<Project>+<StreamName>-CL-<Changelist>-<Postfix>"
+						FString ClKey = TEXT("CL-");
+						int32 ClIndex = BuildInfo.Find(ClKey);
+						if (ClIndex != INDEX_NONE)
+						{
+							// Remove the postfix part after changelist number
+							FString ChangelistStr = BuildInfo.Mid(ClIndex + ClKey.Len());
+							int32 DashIndex = ChangelistStr.Find(TEXT("-"));
+							if (DashIndex != INDEX_NONE)
+							{
+								ChangelistStr = ChangelistStr.Left(DashIndex);
+							}
+							// Return the first backend name matching the changelist, there may be more than one matching backends.
+							int32 Changelist = FCString::Atoi(*ChangelistStr);
+							BuildToBackendsMap.FindOrAdd(Changelist).Add(EpicApp.ToLower());
+						}
+					}
+				}
+			}
+		}
+	}
+
+	// Sort backends into alphabetical order
+	for (auto& Pair : BuildToBackendsMap)
+	{
+		TArray<FString>& Backends = Pair.Value;
+		Backends.Sort();
+	}
+}
+
+
+
+void FBackendRetriever::Completed()
+{
+	bIsGettingBackends = false;
+	if (OnComplete)
+	{
+		OnComplete();
+		OnComplete = nullptr;
+	}
+}
+
+
+
+bool FBackendRetriever::IsGettingBackends() const
+{
+	return bIsGettingBackends;
+}
+
+
+
+bool FBackendRetriever::IsConfigured() const
+{
+	return !BackendQueryUrl.IsEmpty();
+}
+
+const TMap<int32, TArray<FString>>& FBackendRetriever::GetBuildToBackendsMap() const
+{
+	return BuildToBackendsMap;
+}
+
+
+
+const FString& FBackendRetriever::GetErrorMessage() const
+{
+	return ErrorMessage;
+}
diff --git a/Engine/Plugins/Developer/ProjectLauncher/Source/CommonLaunchExtensions/Private/Shared/UGSBuildInfoRetriever.cpp b/Engine/Plugins/Developer/ProjectLauncher/Source/CommonLaunchExtensions/Private/Shared/UGSBuildInfoRetriever.cpp
new file mode 100644
--- /dev/null
+++ b/Engine/Plugins/Developer/ProjectLauncher/Source/CommonLaunchExtensions/Private/Shared/UGSBuildInfoRetriever.cpp
@@ -0,0 +1,252 @@
+// Copyright Epic Games, Inc. All Rights Reserved.
+
+#include "Shared/UGSBuildInfoRetriever.h"
+#include "Misc/EngineVersion.h"
+#include "Misc/ConfigCacheIni.h"
+#include "Dom/JsonObject.h"
+#include "Dom/JsonValue.h"
+#include "Serialization/JsonReader.h"
+#include "Serialization/JsonSerializer.h"
+#include "PlatformHttp.h"
+#include "HttpModule.h"
+#include "Interfaces/IHttpRequest.h"
+#include "Interfaces/IHttpResponse.h"
+
+FUGSBuildInfoRetriever::FUGSBuildInfoRetriever()
+{
+	GConfig->GetString(TEXT("ProjectLauncher"), TEXT("HordeMetadataUrl"), HordeMetadataUrl, GEngineIni);
+}
+
+void FUGSBuildInfoRetriever::GetUGSBuildInfoAsync( const FString& InProjectName, const TArrayView<int32>& InChangelists, FOnUGSBuildInfoRetrievalComplete&& InOnComplete )
+{
+	if (HordeMetadataUrl.IsEmpty())
+	{
+		return;
+	}
+
+	if (bIsGettingUGSBuildInfo)
+	{
+		if (InOnComplete)
+		{
+			InOnComplete();
+		}
+		return;
+	}
+
+	bIsGettingUGSBuildInfo = true;
+	OnComplete = MoveTemp(InOnComplete);
+	ErrorMessage = FString();
+	BuildToUGSBuildInfoMap.Reset();
+
+	// must have a project
+	if (InProjectName.IsEmpty())
+	{
+		ErrorMessage = FString::Printf(TEXT("Failed to query badges. No project specified"));
+		Completed();
+		return;
+	}
+
+
+	// build the UGS badge query URL, limiting to a reasonable maximum length
+	const int MaxUrlLength = 2000;
+	FString Stream = FEngineVersion::Current().GetBranch();
+	FString URL = HordeMetadataUrl;
+
+	URL += TEXT("?stream=") + FPlatformHttp::UrlEncode(Stream.ToLower());
+	URL += TEXT("&project=") + FPlatformHttp::UrlEncode(InProjectName.ToLower());
+	for (int32 Changelist : InChangelists)
+	{
+		URL += FString::Printf( TEXT("&change=%d"), Changelist);
+
+		if (URL.Len() > MaxUrlLength)
+		{
+			break;
+		}
+	}
+
+	// build the UGS query request
+	TSharedRef<IHttpRequest, ESPMode::ThreadSafe> HttpRequest = FHttpModule::Get().CreateRequest();
+	HttpRequest->SetVerb(TEXT("GET"));
+	HttpRequest->SetURL(URL);
+	HttpRequest->OnProcessRequestComplete().BindLambda(
+		[this, ProjectName = InProjectName](FHttpRequestPtr Request, FHttpResponsePtr Response, bool bWasSuccessful)
+		{
+			if (bWasSuccessful && Response->GetResponseCode() == 200)
+			{
+				ParseUGSBuildInfoFromQuery(Response->GetContentAsString(), ProjectName);
+			}
+			else if (Response)
+			{
+				ErrorMessage = FString::Printf(TEXT("Failed to query badges. error code %d: %s"), Response->GetResponseCode(), *Response->GetContentAsString());
+			}
+			else
+			{
+				ErrorMessage = TEXT("Failed to query badges. no response from server");
+			}
+			Completed();
+		}
+	);
+
+	// launch the request
+	if (!HttpRequest->ProcessRequest())
+	{
+		ErrorMessage = FString::Printf(TEXT("Failed to start backend query for stream %s"), *Stream);
+		Completed();
+	}
+}
+
+
+
+void FUGSBuildInfoRetriever::ParseUGSBuildInfoFromQuery(const FString& Text, const FString& ProjectName )
+{
+
+	TSharedPtr<FJsonObject> ResponseJson;
+	if (FJsonSerializer::Deserialize(TJsonReaderFactory<>::Create(Text), ResponseJson) && ResponseJson.IsValid())
+	{
+		const TArray<TSharedPtr<FJsonValue>>* JsonItems;
+		if (ResponseJson->TryGetArrayField(TEXT("items"), JsonItems))
+		{
+			for (const TSharedPtr<FJsonValue>& JsonItem : (*JsonItems))
+			{
+				int32 Changelist;
+				FString Project;
+				const TSharedPtr<FJsonObject>* JsonItemObj;
+				const TArray<TSharedPtr<FJsonValue>>* JsonBadges;
+
+				if ( (JsonItem->TryGetObject(JsonItemObj) &&
+						(*JsonItemObj)->TryGetNumberField(TEXT("change"), Changelist) && 
+						(*JsonItemObj)->TryGetStringField(TEXT("project"), Project) &&
+						(*JsonItemObj)->TryGetArrayField(TEXT("badges"), JsonBadges)))
+				{
+					if (!Project.Equals(ProjectName, ESearchCase::IgnoreCase))
+					{
+						continue;
+					}
+
+					TSharedPtr<FUGSBuildInfo> BuildInfo;
+					if (BuildToUGSBuildInfoMap.Contains(Changelist))
+					{
+						BuildInfo = BuildToUGSBuildInfoMap[Changelist];
+					}
+					else
+					{
+						BuildInfo = BuildToUGSBuildInfoMap.Add(Changelist, MakeShared<FUGSBuildInfo>());
+					}
+
+
+					for (const TSharedPtr<FJsonValue>& JsonBadge : (*JsonBadges))
+					{
+						const TSharedPtr<FJsonObject>* JsonBadgeObj;
+						FString BadgeName;
+						FString BadgeURL;
+						FString BadgeState;
+						if (JsonBadge->TryGetObject(JsonBadgeObj) &&
+							(*JsonBadgeObj)->TryGetStringField(TEXT("name"), BadgeName) &&
+							(*JsonBadgeObj)->TryGetStringField(TEXT("url"), BadgeURL) &&
+							(*JsonBadgeObj)->TryGetStringField(TEXT("state"), BadgeState))
+						{
+							FUGSBadge& Badge = BuildInfo->Badges.AddDefaulted_GetRef();
+							Badge.Name = BadgeName;
+							Badge.URL = BadgeURL;
+							Badge.State = ParseUGSBadgeState(BadgeState);
+						}
+					}
+
+					const TArray<TSharedPtr<FJsonValue>>* JsonUsers;
+					if ((*JsonItemObj)->TryGetArrayField(TEXT("users"), JsonUsers))
+					{
+						BuildInfo->NumUsers += JsonUsers->Num();
+						for (const TSharedPtr<FJsonValue>& JsonUser : (*JsonUsers))
+						{
+							FString UserVote;
+							const TSharedPtr<FJsonObject>* JsonUserObj;
+							if (JsonUser->TryGetObject(JsonUserObj) &&
+								(*JsonUserObj)->TryGetStringField(TEXT("vote"), UserVote))
+							{
+								if (UserVote.Equals(TEXT("CompileSuccess"), ESearchCase::IgnoreCase))
+								{
+									BuildInfo->NumSuccess++;
+								}
+								else if (UserVote.Equals(TEXT("CompileFailure"), ESearchCase::IgnoreCase))
+								{
+									BuildInfo->NumFailed++;
+								}
+								else
+								{
+									ErrorMessage += FString::Printf(TEXT("Unknown UGS badge vote type: %s\n"), *UserVote);
+								}
+							}
+						}
+					}
+				}
+			}
+		}
+	}
+}
+
+
+
+FUGSBuildInfoRetriever::FUGSBadge::EState FUGSBuildInfoRetriever::ParseUGSBadgeState( const FString& State )
+{
+	static const struct
+	{
+		FUGSBadge::EState State;
+		const TCHAR* Text;
+	}
+	States[] = 
+	{
+		{ FUGSBadge::EState::Starting, TEXT("Starting") },
+		{ FUGSBadge::EState::Failure,  TEXT("Failure")  },
+		{ FUGSBadge::EState::Warning,  TEXT("Warning")  },
+		{ FUGSBadge::EState::Success,  TEXT("Success")  },
+		{ FUGSBadge::EState::Skipped,  TEXT("Skipped")  },
+	};
+
+	for (int32 Index = 0; Index < UE_ARRAY_COUNT(States); Index++)
+	{
+		if (State.Equals(States[Index].Text, ESearchCase::IgnoreCase))
+		{
+			return States[Index].State;
+		}
+	}
+
+	return FUGSBadge::EState::Unknown;
+}
+
+
+
+void FUGSBuildInfoRetriever::Completed()
+{
+	bIsGettingUGSBuildInfo = false;
+	if (OnComplete)
+	{
+		OnComplete();
+		OnComplete = nullptr;
+	}
+}
+
+
+
+bool FUGSBuildInfoRetriever::IsGettingUGSBuildInfo() const
+{
+	return bIsGettingUGSBuildInfo;
+}
+
+
+
+bool FUGSBuildInfoRetriever::IsConfigured() const
+{
+	return !HordeMetadataUrl.IsEmpty();
+}
+
+const TMap<int32, TSharedRef<FUGSBuildInfoRetriever::FUGSBuildInfo>>& FUGSBuildInfoRetriever::GetBuildToUGSBuildInfoMap()
+{
+	return BuildToUGSBuildInfoMap;
+}
+
+
+
+const FString& FUGSBuildInfoRetriever::GetErrorMessage() const
+{
+	return ErrorMessage;
+}
diff --git a/Engine/Plugins/Developer/ProjectLauncher/Source/CommonLaunchExtensions/Public/BuildSync/BuildInfoHelper.h b/Engine/Plugins/Developer/ProjectLauncher/Source/CommonLaunchExtensions/Public/BuildSync/BuildInfoHelper.h
new file mode 100644
--- /dev/null
+++ b/Engine/Plugins/Developer/ProjectLauncher/Source/CommonLaunchExtensions/Public/BuildSync/BuildInfoHelper.h
@@ -0,0 +1,111 @@
+// Copyright Epic Games, Inc. All Rights Reserved.
+
+#pragma once
+
+#include "Experimental/BuildListRetriever.h"
+#include "Shared/UGSBuildInfoRetriever.h"
+#include "Shared/BackendRetriever.h"
+#include "Templates/Function.h"
+#include "Templates/SharedPointer.h"
+#include <atomic>
+
+#define UE_API COMMONLAUNCHEXTENSIONS_API
+
+
+class FBuildInfoHelper : public TSharedFromThis<FBuildInfoHelper>
+{
+public:
+	UE_API FBuildInfoHelper();
+	UE_API ~FBuildInfoHelper();
+
+	struct FBuildInfo : public TSharedFromThis<FBuildInfo, ESPMode::ThreadSafe>
+	{
+		TSharedPtr<UE::Zen::Build::FBuildListRetriever::FBuildGroup> Group;
+
+		TSharedPtr<FUGSBuildInfoRetriever::FUGSBuildInfo> UGSBuildInfo;
+		TArray<FString> Backends;
+
+		TSet<FName> Platforms;
+		TMap<FName,TArray<FString>> PlatformToArtifacts;
+	};
+
+
+	UE_API void SetProjectName( const FString& InProjectName );
+	UE_API void SetBuildType( const FString& InBuildType );
+	UE_API void SetBuildsRefreshedHandler(TFunction<void()> OnBuildsRefreshed);
+
+
+	// build service connection
+	UE_API void Connect();
+	UE_API bool IsConnected() const;
+	UE_API bool IsConfigured() const;
+	UE_API FText GetErrorText() const;
+
+	// build list refresh
+	UE_API void Refresh( bool bFullRefresh );
+	UE_API bool IsRefreshing() const;
+
+	// list filtering
+	struct FFilter
+	{
+		FTimespan MaxAge;
+		int32 MaxItems = 50;
+	};
+	void SetFilter( const FFilter& Filter );
+
+
+
+	UE_API TSharedRef<UE::Zen::Build::FBuildServiceInstance> GetServiceInstance() const;
+
+	UE_API const TArray<TSharedPtr<FBuildInfo, ESPMode::ThreadSafe>>& GetBuildInfos() const;
+
+	UE_API const TMap<FString,TSharedRef<UE::Zen::Build::FBuildListRetriever::FBuildType, ESPMode::ThreadSafe>>& GetBuildTypesMap() const;
+
+	UE_API const TSet<FString>& GetAllKnownNamedArtifacts() const;
+
+	UE_API FString GetPlatformNameFromNamedArtifact( const FString& NamedArtifact, const UE::Zen::Build::FBuildServiceInstance::FBuildRecord* BuildRecord = nullptr ) const;
+
+	UE_API static const TCHAR* DefaultBuildType;
+
+	
+
+private:
+
+
+	TFunction<void()> OnBuildsRefreshed;
+	FString ProjectName;
+	FString BuildType;
+	std::atomic<bool> bBuildListRetrieverInitComplete;
+	TSharedPtr<UE::Zen::Build::FBuildListRetriever> BuildListRetriever;
+	TSharedPtr<FBackendRetriever> BackendRetriever;
+	TSharedPtr<FUGSBuildInfoRetriever> UGSBuildInfoRetriever;
+	TSet<FString> KnownPlatformNames;
+
+
+	FFilter Filter;
+
+
+
+
+	// build list details
+	void RefreshBackend();
+	void RefreshUGSBuildInfo();
+	void RefreshBuildGroups();
+	void CreateFilteredBuildGroups( const TArrayView<TSharedRef<UE::Zen::Build::FBuildListRetriever::FBuildGroup>>& BuildGroups ); 
+
+	TArray<TSharedRef<UE::Zen::Build::FBuildListRetriever::FBuildGroup>> FilteredBuildGroups;
+	std::atomic<bool> bBuildGroupsRequestComplete;
+
+	// build info
+	std::atomic<bool> bIsRefreshing = false;
+	void CheckRefreshComplete();
+	void RebuildBuildInfos();
+	TArray<TSharedPtr<FBuildInfo, ESPMode::ThreadSafe>> BuildInfos;
+	TMap<FString,TSharedRef<UE::Zen::Build::FBuildListRetriever::FBuildType, ESPMode::ThreadSafe>> BuildTypesMap;
+	TSet<FString> AllKnownNamedArtifacts;
+
+
+
+};
+
+#undef UE_API
diff --git a/Engine/Plugins/Developer/ProjectLauncher/Source/CommonLaunchExtensions/Public/BuildSync/BuildSyncLaunchExtension.h b/Engine/Plugins/Developer/ProjectLauncher/Source/CommonLaunchExtensions/Public/BuildSync/BuildSyncLaunchExtension.h
new file mode 100644
--- /dev/null
+++ b/Engine/Plugins/Developer/ProjectLauncher/Source/CommonLaunchExtensions/Public/BuildSync/BuildSyncLaunchExtension.h
@@ -0,0 +1,91 @@
+// Copyright Epic Games, Inc. All Rights Reserved.
+
+#pragma once
+
+#include "Extension/CustomUATCommandLaunchExtension.h"
+#include "Experimental/BuildServerInterface.h"
+#include "BuildSync/BuildInfoHelper.h"
+
+#include <atomic>
+
+#define UE_API COMMONLAUNCHEXTENSIONS_API
+
+namespace PlatformInfo
+{
+	struct FTargetPlatformInfo;
+}
+
+class FBuildSyncLaunchExtensionInstance : public ProjectLauncher::FCustomUATCommandLaunchExtensionInstance
+{
+public:
+	FBuildSyncLaunchExtensionInstance( FArgs& InArgs );
+	virtual ~FBuildSyncLaunchExtensionInstance();
+
+	virtual void InternalInitialize() override;
+
+	virtual void CustomizeTree( ProjectLauncher::FLaunchProfileTreeData& ProfileTreeData ) override;
+	virtual void CustomizeUATCommandLine( FString& InOutCommandLine ) override;
+
+	virtual void OnUATCommandAdded( ILauncherProfileUATCommandRef UATCommand ) override;
+
+
+	virtual void OnValidateProfile() override;
+	virtual void OnProjectChanged() override;
+	virtual void OnPropertyChanged() override;
+
+	UE_API TArray<FString> GetSelectedBuildBackends() const;
+
+private:
+	UE::Zen::Build::FBuildServiceInstance::EConnectionState GetBuildServiceConnectionState() const;
+	void OnBuildsRefreshed();
+	void RefreshFilteredBuilds();
+
+	void SetBuildType( FString InBuildType );
+	FString GetBuildType() const;
+	void SetNamedArtifacts( TArray<FString> InNamedArtifacts );
+	TArray<FString> GetNamedArtifacts() const;
+
+	void AutoSelectBestBuild();
+	void AutoSelectNamedArtifacts();
+
+	TSharedRef<SWidget> CreateBuildStorageConnectionWidget();
+	TSharedRef<SWidget> CreateBuildFilterWidget();
+	TSharedRef<SWidget> CreateBuildTypeSelectionWidget();
+	TSharedRef<SWidget> CreateNamedArtifactSelectionWidget();
+
+	void SetSelectedBuild(TSharedPtr<FBuildInfoHelper::FBuildInfo> Build);
+	TSharedPtr<FBuildInfoHelper::FBuildInfo> GetSelectedBuild() const;
+	TArray<FString> GetAvailableArtifactsForSelectedBuild() const;
+
+	bool GetBuildSuitability( TSharedPtr<FBuildInfoHelper::FBuildInfo> Build, FText* OutReason ) const;
+	void CacheUnsuitableBuilds();
+	
+	TSharedPtr<FBuildInfoHelper::FBuildInfo> SelectedBuild;
+
+	TSharedPtr<FBuildInfoHelper> BuildInfo;
+	TSharedRef<class FBuildSyncLaunchExtension> Owner;
+	TSharedPtr<class SBuildSyncBuildCombo> BuildCombo;
+	TArray<FString> CachedNamedArtifacts;
+	TMap<TSharedPtr<FBuildInfoHelper::FBuildInfo>, FText> CachedUnsuitableBuilds;
+	TArray<TSharedPtr<FBuildInfoHelper::FBuildInfo>> FilteredBuildInfos;
+
+	bool bShowPreflights = false;
+	bool bShowUnsuitable = false;
+	FString FilterText;
+
+	TSet<const PlatformInfo::FTargetPlatformInfo*> GetCookPlatforms() const;
+	TSet<const PlatformInfo::FTargetPlatformInfo*> CachedCookPlatforms;
+};
+
+
+class FBuildSyncLaunchExtension : public ProjectLauncher::FCustomUATCommandLaunchExtension
+{
+public:
+	FBuildSyncLaunchExtension();
+
+	virtual TSharedPtr<ProjectLauncher::FLaunchExtensionInstance> CreateInstanceForProfile( ProjectLauncher::FLaunchExtensionInstance::FArgs& InArgs ) override;
+	virtual const TCHAR* GetInternalName() const override;
+	virtual FText GetDisplayName() const override;
+};
+
+#undef UE_API
\ No newline at end of file
diff --git a/Engine/Plugins/Developer/ProjectLauncher/Source/CommonLaunchExtensions/Public/BuildSync/SBuildSyncBuildCombo.h b/Engine/Plugins/Developer/ProjectLauncher/Source/CommonLaunchExtensions/Public/BuildSync/SBuildSyncBuildCombo.h
new file mode 100644
--- /dev/null
+++ b/Engine/Plugins/Developer/ProjectLauncher/Source/CommonLaunchExtensions/Public/BuildSync/SBuildSyncBuildCombo.h
@@ -0,0 +1,61 @@
+// Copyright Epic Games, Inc. All Rights Reserved.
+
+#pragma once
+
+#include "CoreMinimal.h"
+#include "Widgets/SCompoundWidget.h"
+#include "Widgets/Views/SListView.h"
+#include "Widgets/Input/SComboButton.h"
+#include "Styling/CoreStyle.h"
+#include "Styling/SlateWidgetStyleAsset.h"
+#include "Fonts/SlateFontInfo.h"
+#include "BuildSync/BuildInfoHelper.h"
+
+#define UE_API COMMONLAUNCHEXTENSIONS_API
+
+template<typename ItemType> class SComboBox;
+
+
+class SBuildSyncBuildCombo
+	: public SCompoundWidget
+{
+public:
+	DECLARE_DELEGATE_OneParam(FOnSelectionChanged, TSharedPtr<FBuildInfoHelper::FBuildInfo> );
+	DECLARE_DELEGATE_RetVal_TwoParams(bool, FGetItemSuitability, TSharedPtr<FBuildInfoHelper::FBuildInfo>, FText* );
+
+	SLATE_BEGIN_ARGS(SBuildSyncBuildCombo)
+		: _OptionsSource()
+		, _TextStyle( &FCoreStyle::Get().GetWidgetStyle<FTextBlockStyle>( "NormalText" ) )
+		{}
+		SLATE_ARGUMENT(const TArray<TSharedPtr<FBuildInfoHelper::FBuildInfo>>*, OptionsSource)
+		SLATE_ARGUMENT(TSharedPtr<SWidget>, FilterWidget)
+		SLATE_EVENT(FOnSelectionChanged, OnSelectionChanged);
+		SLATE_ATTRIBUTE(TSharedPtr<FBuildInfoHelper::FBuildInfo>, SelectedItem);
+		SLATE_ATTRIBUTE(FSlateFontInfo, Font)
+		SLATE_ATTRIBUTE(TSet<FName>, RequiredPlatforms)
+		SLATE_EVENT(FGetItemSuitability, GetItemSuitability)
+		SLATE_STYLE_ARGUMENT( FTextBlockStyle, TextStyle )
+	SLATE_END_ARGS()
+
+public:
+	UE_API void Construct(const FArguments& InArgs);
+	UE_API void RefreshOptions();
+
+protected:
+	const TArray<TSharedPtr<FBuildInfoHelper::FBuildInfo>>* OptionsSource;
+	TAttribute<TSharedPtr<FBuildInfoHelper::FBuildInfo>> SelectedItem;
+	TAttribute<TSet<FName>> RequiredPlatforms;
+	FOnSelectionChanged OnSelectionChanged;
+	FGetItemSuitability GetItemSuitability;
+
+	TSharedRef<ITableRow> OnGenerateBuildInfoRow( TSharedPtr<FBuildInfoHelper::FBuildInfo> BuildInfo, const TSharedRef<STableViewBase>& OwnerTable) const;
+	void OnBuildInfoSelectionChanged( TSharedPtr<FBuildInfoHelper::FBuildInfo> BuildInfo, ESelectInfo::Type InSelectInfo );
+	FText GetSelectedBuildInfoName() const;
+
+	TSharedPtr<SComboButton> ComboButton;
+	TSharedPtr<SListView<TSharedPtr<FBuildInfoHelper::FBuildInfo>>> BuildInfosListView;
+
+	FText GetAge( const FDateTime& DateTime ) const;
+};
+
+#undef UE_API
diff --git a/Engine/Plugins/Developer/ProjectLauncher/Source/CommonLaunchExtensions/Public/Shared/BackendRetriever.h b/Engine/Plugins/Developer/ProjectLauncher/Source/CommonLaunchExtensions/Public/Shared/BackendRetriever.h
new file mode 100644
--- /dev/null
+++ b/Engine/Plugins/Developer/ProjectLauncher/Source/CommonLaunchExtensions/Public/Shared/BackendRetriever.h
@@ -0,0 +1,36 @@
+// Copyright Epic Games, Inc. All Rights Reserved.
+
+#pragma once
+
+#include "Templates/Function.h"
+#include "Templates/SharedPointer.h"
+#include "Containers/Map.h"
+#include <atomic>
+
+
+class FBackendRetriever : public TSharedFromThis<FBackendRetriever>
+{
+public:
+	typedef TUniqueFunction<void()> FOnBackendRetrievalComplete;
+
+	FBackendRetriever();
+
+	void GetBackendsAsync(FOnBackendRetrievalComplete&& InOnComplete);
+	bool IsGettingBackends() const;
+	bool IsConfigured() const;
+
+	const TMap<int32, TArray<FString>>& GetBuildToBackendsMap() const;
+	const FString& GetErrorMessage() const;
+
+private:
+
+	void ParseBackendFromOutput(const FString& QueryResult);
+	void Completed();
+
+	FOnBackendRetrievalComplete OnComplete;
+	std::atomic<bool> bIsGettingBackends = false;
+
+	TMap<int32, TArray<FString>> BuildToBackendsMap;
+	FString ErrorMessage;
+	FString BackendQueryUrl;
+};
diff --git a/Engine/Plugins/Developer/ProjectLauncher/Source/CommonLaunchExtensions/Public/Shared/UGSBuildInfoRetriever.h b/Engine/Plugins/Developer/ProjectLauncher/Source/CommonLaunchExtensions/Public/Shared/UGSBuildInfoRetriever.h
new file mode 100644
--- /dev/null
+++ b/Engine/Plugins/Developer/ProjectLauncher/Source/CommonLaunchExtensions/Public/Shared/UGSBuildInfoRetriever.h
@@ -0,0 +1,66 @@
+// Copyright Epic Games, Inc. All Rights Reserved.
+
+#pragma once
+
+#include "Templates/Function.h"
+#include "Templates/SharedPointer.h"
+#include "Containers/ArrayView.h"
+#include "Containers/Map.h"
+#include <atomic>
+
+
+class FUGSBuildInfoRetriever : public TSharedFromThis<FUGSBuildInfoRetriever>
+{
+public:
+	struct FUGSBadge
+	{
+		enum class EState : uint8
+		{
+			Starting,
+			Failure,
+			Warning,
+			Success,
+			Skipped,
+
+			Unknown,
+			MAX
+		};
+
+		FString Name;
+		FString URL;
+		EState State;
+	};
+
+	struct FUGSBuildInfo : public TSharedFromThis<FUGSBuildInfo, ESPMode::ThreadSafe>
+	{
+		int32 NumUsers = 0;
+		int32 NumSuccess = 0;
+		int32 NumFailed = 0;
+
+		TArray<FUGSBadge> Badges;
+	};
+
+	typedef TUniqueFunction<void()> FOnUGSBuildInfoRetrievalComplete;
+
+	FUGSBuildInfoRetriever();
+
+	void GetUGSBuildInfoAsync( const FString& InProjectName, const TArrayView<int32>& InChangelists, FOnUGSBuildInfoRetrievalComplete&& InOnComplete );
+	bool IsGettingUGSBuildInfo() const;
+	bool IsConfigured() const;
+
+	const TMap<int32, TSharedRef<FUGSBuildInfo>>& GetBuildToUGSBuildInfoMap();
+	const FString& GetErrorMessage() const;
+
+private:
+
+	void ParseUGSBuildInfoFromQuery(const FString& Text, const FString& ProjectName);
+	FUGSBadge::EState ParseUGSBadgeState(const FString& Text);
+	void Completed();
+
+	FOnUGSBuildInfoRetrievalComplete OnComplete;
+	std::atomic<bool> bIsGettingUGSBuildInfo = false;
+
+	TMap<int32, TSharedRef<FUGSBuildInfo>> BuildToUGSBuildInfoMap;
+	FString ErrorMessage;
+	FString HordeMetadataUrl;
+};
