device selection widgets show previously-selected devices before they are discovered, or after they are disconnected

#rnx

#Synced-CL 50642529

diff --git a/Engine/Plugins/Developer/ProjectLauncher/Source/ProjectLauncher/Private/ProfileTree/GenericProfileTreeBuilder.cpp b/Engine/Plugins/Developer/ProjectLauncher/Source/ProjectLauncher/Private/ProfileTree/GenericProfileTreeBuilder.cpp
--- a/Engine/Plugins/Developer/ProjectLauncher/Source/ProjectLauncher/Private/ProfileTree/GenericProfileTreeBuilder.cpp
+++ b/Engine/Plugins/Developer/ProjectLauncher/Source/ProjectLauncher/Private/ProfileTree/GenericProfileTreeBuilder.cpp
@@ -1331,7 +1331,6 @@ namespace ProjectLauncher
 					.HandleColor( FAppStyle::Get().GetSlateColor("Colors.Secondary").GetSpecifiedColor() )
 					[
 						SAssignNew(DeployDeviceListView, SCustomLaunchDeviceListView)
-						.OnDeviceRemoved( this, &FGenericProfileTreeBuilder::FBuildCookRun::OnDeviceRemoved)
 						.OnSelectionChanged(this, &FGenericProfileTreeBuilder::FBuildCookRun::SetDeployDeviceIDs)
 						.SelectedDevices(this, &FGenericProfileTreeBuilder::FBuildCookRun::GetDeployDeviceIDs)
 						.Platforms(this, &FGenericProfileTreeBuilder::FBuildCookRun::GetSelectedPlatforms)
@@ -1383,11 +1382,6 @@ namespace ProjectLauncher
 		Owner.TreeData->RequestTreeRefresh();
 	}
 
-	void FGenericProfileTreeBuilder::FBuildCookRun::OnDeviceRemoved(FString DeviceID)
-	{
-		CachedDeployDeviceIDs.Remove(DeviceID);
-	}
-
 	FGenericProfileTreeBuilder::FBuildCookRun::EDeployDeviceOption FGenericProfileTreeBuilder::FBuildCookRun::GetDeployDeviceOption() const
 	{
 		return (BuildCookRun->GetDefaultDeployPlatform() == NAME_None) ? EDeployDeviceOption::Selected : EDeployDeviceOption::Default;;
diff --git a/Engine/Plugins/Developer/ProjectLauncher/Source/ProjectLauncher/Private/Widgets/Profiles/SCustomLaunchCustomProfileSelector.cpp b/Engine/Plugins/Developer/ProjectLauncher/Source/ProjectLauncher/Private/Widgets/Profiles/SCustomLaunchCustomProfileSelector.cpp
--- a/Engine/Plugins/Developer/ProjectLauncher/Source/ProjectLauncher/Private/Widgets/Profiles/SCustomLaunchCustomProfileSelector.cpp
+++ b/Engine/Plugins/Developer/ProjectLauncher/Source/ProjectLauncher/Private/Widgets/Profiles/SCustomLaunchCustomProfileSelector.cpp
@@ -206,7 +206,6 @@ TSharedRef<ITableRow> SCustomLaunchCustomProfileSelector::GenerateCustomProfileR
 			.VAlign(VAlign_Center)
 			[
 				SNew(SCustomLaunchDeviceCombo)
-				.OnDeviceRemoved(this, &SCustomLaunchCustomProfileSelector::OnDeviceRemoved, Profile)
 				.OnSelectionChanged(this, &SCustomLaunchCustomProfileSelector::SetSelectedDevices, Profile)
 				.SelectedDevices(this, &SCustomLaunchCustomProfileSelector::GetSelectedDevices, Profile)
 				.AllPlatforms(true)
@@ -524,26 +523,6 @@ void SCustomLaunchCustomProfileSelector::OnCustomProfileEditDescriptionClicked(I
 	OnProfileEditDescription.ExecuteIfBound(Profile);
 }
 
-void SCustomLaunchCustomProfileSelector::OnDeviceRemoved( const FString DeviceID, ILauncherProfilePtr Profile)
-{
-	//@todo: not sure whether this is needed?
-
-	bool bModified = false;
-	for (ILauncherProfileBuildCookRunRef BuildCookRun : Profile->GetBuildCookRunCommands())
-	{
-		if (ensure(BuildCookRun->GetDeployedDeviceGroup() != nullptr) && BuildCookRun->GetDeployedDeviceGroup()->GetDeviceIDs().Contains(DeviceID))
-		{
-			BuildCookRun->GetDeployedDeviceGroup()->RemoveDevice(DeviceID);
-			bModified = true;
-		}
-	}
-
-	if (bModified)
-	{
-		OnProfileModified.ExecuteIfBound(Profile);
-	}
-}
-
 void SCustomLaunchCustomProfileSelector::SetSelectedDevices( const TArray<FString> DeviceIDs, ILauncherProfilePtr Profile)
 {
 	check(Model->IsBasicLaunchProfile(Profile));
diff --git a/Engine/Plugins/Developer/ProjectLauncher/Source/ProjectLauncher/Private/Widgets/Profiles/SCustomLaunchCustomProfileSelector.h b/Engine/Plugins/Developer/ProjectLauncher/Source/ProjectLauncher/Private/Widgets/Profiles/SCustomLaunchCustomProfileSelector.h
--- a/Engine/Plugins/Developer/ProjectLauncher/Source/ProjectLauncher/Private/Widgets/Profiles/SCustomLaunchCustomProfileSelector.h
+++ b/Engine/Plugins/Developer/ProjectLauncher/Source/ProjectLauncher/Private/Widgets/Profiles/SCustomLaunchCustomProfileSelector.h
@@ -75,7 +75,6 @@ private:
 
 	void RefreshCustomProfileList();
 
-	void OnDeviceRemoved(const FString DeviceID, ILauncherProfilePtr Profile);
 	void SetSelectedDevices(const TArray<FString> DeviceIDs, ILauncherProfilePtr Profile );
 	TArray<FString> GetSelectedDevices(const ILauncherProfilePtr Profile) const;
 
diff --git a/Engine/Plugins/Developer/ProjectLauncher/Source/ProjectLauncher/Private/Widgets/Shared/SCustomLaunchDeviceCombo.cpp b/Engine/Plugins/Developer/ProjectLauncher/Source/ProjectLauncher/Private/Widgets/Shared/SCustomLaunchDeviceCombo.cpp
--- a/Engine/Plugins/Developer/ProjectLauncher/Source/ProjectLauncher/Private/Widgets/Shared/SCustomLaunchDeviceCombo.cpp
+++ b/Engine/Plugins/Developer/ProjectLauncher/Source/ProjectLauncher/Private/Widgets/Shared/SCustomLaunchDeviceCombo.cpp
@@ -13,6 +13,7 @@
 #include "ITargetDeviceProxy.h"
 #include "ITargetDeviceProxyManager.h"
 #include "PlatformInfo.h"
+#include "Model/ProjectLauncherModel.h"
 
 #define LOCTEXT_NAMESPACE "SCustomLaunchDeviceCombo"
 
@@ -20,7 +21,6 @@
 BEGIN_SLATE_FUNCTION_BUILD_OPTIMIZATION
 void SCustomLaunchDeviceCombo::Construct(const FArguments& InArgs)
 {
-	OnDeviceRemoved = InArgs._OnDeviceRemoved;
 	OnSelectionChanged = InArgs._OnSelectionChanged;
 	SelectedDevices = InArgs._SelectedDevices;
 	Platforms = InArgs._Platforms;
@@ -33,10 +33,10 @@ void SCustomLaunchDeviceCombo::Construct(const FArguments& InArgs)
 		+SHorizontalBox::Slot()
 		.AutoWidth()		
 		[
-			SAssignNew(DeviceProxyComboBox, SComboBox<TSharedPtr<ITargetDeviceProxy>>)
-			.OptionsSource(&DeviceProxyList)
-			.OnGenerateWidget(this, &SCustomLaunchDeviceCombo::GenerateDeviceProxyListWidget)
-			.OnSelectionChanged( this, &SCustomLaunchDeviceCombo::OnDeviceProxySelectionChanged)
+			SAssignNew(DeviceComboBox, SComboBox<TSharedPtr<FString>>)
+			.OptionsSource(&AllDeviceIDs)
+			.OnGenerateWidget(this, &SCustomLaunchDeviceCombo::GenerateDeviceListWidget)
+			.OnSelectionChanged( this, &SCustomLaunchDeviceCombo::OnDeviceSelectionChanged)
 			[
 				SNew(SHorizontalBox)
 					
@@ -47,7 +47,7 @@ void SCustomLaunchDeviceCombo::Construct(const FArguments& InArgs)
 				[
 					SNew(SImage)
 					.DesiredSizeOverride(FVector2D(16,16))
-					.Image(this, &SCustomLaunchDeviceCombo::GetSelectedDeviceProxyBrush)
+					.Image(this, &SCustomLaunchDeviceCombo::GetSelectedDeviceBrush)
 				]
 
 				+SHorizontalBox::Slot()
@@ -57,7 +57,7 @@ void SCustomLaunchDeviceCombo::Construct(const FArguments& InArgs)
 				[
 					SNew(STextBlock)
 					//.TextStyle( FAppStyle::Get(), "SmallText")
-					.Text(this, &SCustomLaunchDeviceCombo::GetSelectedDeviceProxyName)
+					.Text(this, &SCustomLaunchDeviceCombo::GetSelectedDeviceName)
 				]
 			]
 		]
@@ -73,9 +73,10 @@ END_SLATE_FUNCTION_BUILD_OPTIMIZATION
 
 
 BEGIN_SLATE_FUNCTION_BUILD_OPTIMIZATION
-TSharedRef<SWidget> SCustomLaunchDeviceCombo::GenerateDeviceProxyListWidget( TSharedPtr<ITargetDeviceProxy> DeviceProxy ) const
+TSharedRef<SWidget> SCustomLaunchDeviceCombo::GenerateDeviceListWidget( TSharedPtr<FString> DeviceID ) const
 {
-	const PlatformInfo::FTargetPlatformInfo* PlatformInfo = PlatformInfo::FindPlatformInfo(*DeviceProxy->GetTargetPlatformName(NAME_None));
+	const PlatformInfo::FTargetPlatformInfo* PlatformInfo = GetPlatformInfoForDeviceID(*DeviceID);
+	const TSharedPtr<ITargetDeviceProxy> DeviceProxy = GetDeviceProxyManager()->FindProxyDeviceForTargetDevice(*DeviceID);
 
 	return SNew(SHorizontalBox)
 
@@ -97,14 +98,15 @@ TSharedRef<SWidget> SCustomLaunchDeviceCombo::GenerateDeviceProxyListWidget( TSh
 		.Padding(4,0)
 		[
 			SNew(STextBlock)
-			.Text(FText::FromString(*DeviceProxy->GetName()))
+			.Text(GetDisplayNameForDevice(DeviceID))
+			.ColorAndOpacity( DeviceProxy.IsValid() ? GetForegroundColor() : GetDisabledForegroundColor() )
 		]
 	;
 }
 END_SLATE_FUNCTION_BUILD_OPTIMIZATION
 
 
-const FSlateBrush* SCustomLaunchDeviceCombo::GetSelectedDeviceProxyBrush() const
+const FSlateBrush* SCustomLaunchDeviceCombo::GetSelectedDeviceBrush() const
 {
 	TArray<FString> DeviceIDs = SelectedDevices.Get();
 
@@ -129,7 +131,7 @@ const FSlateBrush* SCustomLaunchDeviceCombo::GetSelectedDeviceProxyBrush() const
 	return FStyleDefaults::GetNoBrush();
 }
 
-FText SCustomLaunchDeviceCombo::GetSelectedDeviceProxyName() const
+FText SCustomLaunchDeviceCombo::GetSelectedDeviceName() const
 {
 	TArray<FString> DeviceIDs = SelectedDevices.Get();
 
@@ -150,14 +152,11 @@ FText SCustomLaunchDeviceCombo::GetSelectedDeviceProxyName() const
 	return LOCTEXT("NoDevice", "(no device)");
 }
 
-void SCustomLaunchDeviceCombo::OnDeviceProxySelectionChanged(TSharedPtr<ITargetDeviceProxy> DeviceProxy, ESelectInfo::Type InSelectInfo)
+void SCustomLaunchDeviceCombo::OnDeviceSelectionChanged(TSharedPtr<FString> DeviceID, ESelectInfo::Type InSelectInfo)
 {
 	TArray<FString> DeviceIDs;
 
-	if (DeviceProxy != nullptr)
-	{
-		DeviceIDs.Add(DeviceProxy->GetTargetDeviceId(NAME_None));
-	}
+	DeviceIDs.Add(*DeviceID);
 
 	OnSelectionChanged.ExecuteIfBound(DeviceIDs);
 }
diff --git a/Engine/Plugins/Developer/ProjectLauncher/Source/ProjectLauncher/Private/Widgets/Shared/SCustomLaunchDeviceListView.cpp b/Engine/Plugins/Developer/ProjectLauncher/Source/ProjectLauncher/Private/Widgets/Shared/SCustomLaunchDeviceListView.cpp
--- a/Engine/Plugins/Developer/ProjectLauncher/Source/ProjectLauncher/Private/Widgets/Shared/SCustomLaunchDeviceListView.cpp
+++ b/Engine/Plugins/Developer/ProjectLauncher/Source/ProjectLauncher/Private/Widgets/Shared/SCustomLaunchDeviceListView.cpp
@@ -12,7 +12,9 @@
 #include "Styling/AppStyle.h"
 #include "Styling/ProjectLauncherStyle.h"
 #include "ITargetDeviceProxy.h"
+#include "ITargetDeviceProxyManager.h"
 #include "PlatformInfo.h"
+#include "Model/ProjectLauncherModel.h"
 
 #define LOCTEXT_NAMESPACE "SCustomLaunchDeviceListView"
 
@@ -20,7 +22,6 @@
 BEGIN_SLATE_FUNCTION_BUILD_OPTIMIZATION
 void SCustomLaunchDeviceListView::Construct(const FArguments& InArgs)
 {
-	OnDeviceRemoved = InArgs._OnDeviceRemoved;
 	OnSelectionChanged = InArgs._OnSelectionChanged;
 	SelectedDevices = InArgs._SelectedDevices;
 	Platforms = InArgs._Platforms;
@@ -33,9 +34,9 @@ void SCustomLaunchDeviceListView::Construct(const FArguments& InArgs)
 		.Padding(2)
 		.BorderImage(FAppStyle::GetBrush("Brushes.Background"))
 		[
-			SAssignNew(DeviceProxyListView, SListView<TSharedPtr<ITargetDeviceProxy>>)
-			.ListItemsSource(&DeviceProxyList)
-			.OnGenerateRow(this, &SCustomLaunchDeviceListView::GenerateDeviceProxyRow)
+			SAssignNew(DeviceListView, SListView<TSharedPtr<FString>>)
+			.ListItemsSource(&AllDeviceIDs)
+			.OnGenerateRow(this, &SCustomLaunchDeviceListView::GenerateDeviceRow)
 		]
 	];
 
@@ -47,13 +48,14 @@ END_SLATE_FUNCTION_BUILD_OPTIMIZATION
 
 
 BEGIN_SLATE_FUNCTION_BUILD_OPTIMIZATION
-TSharedRef<ITableRow> SCustomLaunchDeviceListView::GenerateDeviceProxyRow(TSharedPtr<ITargetDeviceProxy> DeviceProxy, const TSharedRef<STableViewBase>& OwnerTable)
+TSharedRef<ITableRow> SCustomLaunchDeviceListView::GenerateDeviceRow(TSharedPtr<FString> DeviceID, const TSharedRef<STableViewBase>& OwnerTable)
 {
-	const PlatformInfo::FTargetPlatformInfo* PlatformInfo = PlatformInfo::FindPlatformInfo(*DeviceProxy->GetTargetPlatformName(NAME_None));
+	const PlatformInfo::FTargetPlatformInfo* PlatformInfo = GetPlatformInfoForDeviceID(*DeviceID);
+	const TSharedPtr<ITargetDeviceProxy> DeviceProxy = GetDeviceProxyManager()->FindProxyDeviceForTargetDevice(*DeviceID);
 
 	TSharedPtr<SWidget> DeviceSelectorCheckbox = SNew(SCheckBox)
-	.IsChecked(this, &SCustomLaunchDeviceListView::IsDeviceProxyChecked, DeviceProxy)
-	.OnCheckStateChanged(this, &SCustomLaunchDeviceListView::OnDeviceProxyCheckStateChanged, DeviceProxy)
+	.IsChecked(this, &SCustomLaunchDeviceListView::IsDeviceChecked, DeviceID)
+	.OnCheckStateChanged(this, &SCustomLaunchDeviceListView::OnDeviceCheckStateChanged, DeviceID)
 	.Style(FAppStyle::Get(), bSingleSelect ? "RadioButton" : "Checkbox")
 	[
 		SNew(SHorizontalBox)
@@ -76,11 +78,12 @@ TSharedRef<ITableRow> SCustomLaunchDeviceListView::GenerateDeviceProxyRow(TShare
 		.Padding(4,0)
 		[
 			SNew(STextBlock)
-			.Text(FText::FromString(*DeviceProxy->GetName()))
+			.Text(GetDisplayNameForDevice(DeviceID))
+			.ColorAndOpacity( DeviceProxy.IsValid() ? FSlateColor::UseForeground() : FSlateColor::UseSubduedForeground() )
 		]
 	];
 
-	return SNew(STableRow<TSharedPtr<ITargetDeviceProxy>>, OwnerTable)
+	return SNew(STableRow<TSharedPtr<FString>>, OwnerTable)
 	.Padding(FMargin(4,1))
 	[
 		DeviceSelectorCheckbox.ToSharedRef()
@@ -89,11 +92,9 @@ TSharedRef<ITableRow> SCustomLaunchDeviceListView::GenerateDeviceProxyRow(TShare
 END_SLATE_FUNCTION_BUILD_OPTIMIZATION
 
 
-ECheckBoxState SCustomLaunchDeviceListView::IsDeviceProxyChecked(TSharedPtr<ITargetDeviceProxy> DeviceProxy) const
+ECheckBoxState SCustomLaunchDeviceListView::IsDeviceChecked(TSharedPtr<FString> DeviceID) const
 {
-	FString DeviceID = DeviceProxy->GetTargetDeviceId(NAME_None);
-
-	if (SelectedDevices.Get().Contains(DeviceID))
+	if (SelectedDevices.Get().Contains(*DeviceID))
 	{
 		return ECheckBoxState::Checked;
 	}
@@ -105,10 +106,8 @@ ECheckBoxState SCustomLaunchDeviceListView::IsDeviceProxyChecked(TSharedPtr<ITar
 
 
 
-void SCustomLaunchDeviceListView::OnDeviceProxyCheckStateChanged(ECheckBoxState NewState, TSharedPtr<ITargetDeviceProxy> DeviceProxy)
+void SCustomLaunchDeviceListView::OnDeviceCheckStateChanged(ECheckBoxState NewState, TSharedPtr<FString> DeviceID)
 {
-	FString DeviceID = DeviceProxy->GetTargetDeviceId(NAME_None);
-
 	TArray<FString> Devices;
 	if (!bSingleSelect)
 	{
@@ -122,11 +121,11 @@ void SCustomLaunchDeviceListView::OnDeviceProxyCheckStateChanged(ECheckBoxState
 
 	if (NewState == ECheckBoxState::Checked)
 	{
-		Devices.Add(DeviceID);
+		Devices.Add(*DeviceID);
 	}
 	else
 	{
-		Devices.Remove(DeviceID);
+		Devices.Remove(*DeviceID);
 	}
 
 	OnSelectionChanged.ExecuteIfBound(Devices);
@@ -135,7 +134,7 @@ void SCustomLaunchDeviceListView::OnDeviceProxyCheckStateChanged(ECheckBoxState
 
 void SCustomLaunchDeviceListView::OnDeviceListRefreshed()
 {
-	DeviceProxyListView->RequestListRefresh();
+	DeviceListView->RequestListRefresh();
 }
 
 
diff --git a/Engine/Plugins/Developer/ProjectLauncher/Source/ProjectLauncher/Private/Widgets/Shared/SCustomLaunchDeviceWidgetBase.cpp b/Engine/Plugins/Developer/ProjectLauncher/Source/ProjectLauncher/Private/Widgets/Shared/SCustomLaunchDeviceWidgetBase.cpp
--- a/Engine/Plugins/Developer/ProjectLauncher/Source/ProjectLauncher/Private/Widgets/Shared/SCustomLaunchDeviceWidgetBase.cpp
+++ b/Engine/Plugins/Developer/ProjectLauncher/Source/ProjectLauncher/Private/Widgets/Shared/SCustomLaunchDeviceWidgetBase.cpp
@@ -7,6 +7,8 @@
 #include "ITargetDeviceServicesModule.h"
 #include "PlatformInfo.h"
 #include "Modules/ModuleManager.h"
+#include "Algo/StableSort.h"
+#include "Model/ProjectLauncherModel.h"
 
 void SCustomLaunchDeviceWidgetBase::Construct()
 {
@@ -35,9 +37,6 @@ void SCustomLaunchDeviceWidgetBase::OnDeviceProxyAdded(const TSharedRef<ITargetD
 
 void SCustomLaunchDeviceWidgetBase::OnDeviceProxyRemoved(const TSharedRef<ITargetDeviceProxy>& DeviceProxy)
 {
-	FString DeviceID = DeviceProxy->GetTargetDeviceId(NAME_None);
-	OnDeviceRemoved.ExecuteIfBound(DeviceID);
-
 	RefreshDeviceList();
 }
 
@@ -45,8 +44,9 @@ void SCustomLaunchDeviceWidgetBase::OnDeviceProxyRemoved(const TSharedRef<ITarge
 
 void SCustomLaunchDeviceWidgetBase::RefreshDeviceList()
 {
-	DeviceProxyList.Reset();
-
+	// gather device proxies for all required platforms
+	TArray<FString> RequiredPlatforms = Platforms.Get();
+	TArray<TSharedPtr<ITargetDeviceProxy>> DeviceProxyList;
 	if (bAllPlatforms)
 	{
 		for (const PlatformInfo::FTargetPlatformInfo* PlatformInfo : PlatformInfo::GetVanillaPlatformInfoArray())
@@ -62,7 +62,7 @@ void SCustomLaunchDeviceWidgetBase::RefreshDeviceList()
 	}
 	else
 	{
-		for (const FString& Platform : Platforms.Get())
+		for (const FString& Platform : RequiredPlatforms)
 		{
 			const PlatformInfo::FTargetPlatformInfo* PlatformInfo = PlatformInfo::FindPlatformInfo(FName(Platform));
 			if (PlatformInfo != nullptr)
@@ -78,6 +78,42 @@ void SCustomLaunchDeviceWidgetBase::RefreshDeviceList()
 		}
 	}
 
+	// remember all selected devices we've seeen
+	for (const FString& DeviceID : SelectedDevices.Get())
+	{
+		KnownDeviceIDs.Add(DeviceID);
+	}
+
+	// add the device proxies
+	TArray<FString> DeviceIDs;
+	for (const TSharedPtr<ITargetDeviceProxy>& DeviceProxy : DeviceProxyList)
+	{
+		FString DeviceID = DeviceProxy->GetTargetDeviceId(NAME_None);
+		DeviceIDs.Add(DeviceID);
+		KnownDeviceIDs.Add(DeviceID);
+	}
+
+	// add all known devices (these will include devices that have been disconnected or not discovered yet)
+	for (const FString& DeviceID : KnownDeviceIDs)
+	{
+		if (!DeviceIDs.Contains(DeviceID))
+		{
+			const PlatformInfo::FTargetPlatformInfo* PlatformInfo = GetPlatformInfoForDeviceID(DeviceID);
+			if (bAllPlatforms || PlatformInfo == nullptr || RequiredPlatforms.Contains(PlatformInfo->Name))
+			{
+				DeviceIDs.Add(DeviceID);
+			}
+		}
+	}
+	
+	// finalize the sorted device id list (note: these will be sorted by platform@devicename, so each platforms' devices will be grouped together)
+	Algo::StableSort(DeviceIDs);
+	AllDeviceIDs.Reset();
+	for (const FString& DeviceID : DeviceIDs)
+	{
+		AllDeviceIDs.Add(MakeShared<FString>(DeviceID));
+	}
+
 	OnDeviceListRefreshed();
 }
 
@@ -102,6 +138,7 @@ void SCustomLaunchDeviceWidgetBase::OnSelectedPlatformChanged()
 	}
 
 	// remove any devices that belong to other platforms
+	KnownDeviceIDs.Reset();
 	TArray<FString> ValidSelectedDevices = SelectedDevices.Get();
 	int32 InvalidCount = ValidSelectedDevices.RemoveAll( [ValidPlatformDevices]( const FString& Device )
 	{
@@ -125,3 +162,38 @@ const TSharedRef<ITargetDeviceProxyManager> SCustomLaunchDeviceWidgetBase::GetDe
 }
 
 
+const PlatformInfo::FTargetPlatformInfo* SCustomLaunchDeviceWidgetBase::GetPlatformInfoForDeviceID( const FString& DeviceID ) const
+{
+	const TSharedPtr<ITargetDeviceProxy> DeviceProxy = GetDeviceProxyManager()->FindProxyDeviceForTargetDevice(DeviceID);
+	if (DeviceProxy.IsValid())
+	{
+		return PlatformInfo::FindPlatformInfo(*DeviceProxy->GetTargetPlatformName(NAME_None));
+	}
+
+	FString PlatformName;
+	if (DeviceID.Split(TEXT("@"), &PlatformName, nullptr))
+	{
+		return PlatformInfo::FindPlatformInfo(FName(*PlatformName));
+	}
+
+	return nullptr;
+}
+
+
+FText SCustomLaunchDeviceWidgetBase::GetDisplayNameForDevice( const TSharedPtr<FString>& DeviceID ) const
+{
+	const TSharedPtr<ITargetDeviceProxy> DeviceProxy = GetDeviceProxyManager()->FindProxyDeviceForTargetDevice(*DeviceID);
+	if (DeviceProxy.IsValid())
+	{
+		return FText::FromString(DeviceProxy->GetName());
+	}
+
+	FString DeviceName;
+	if (!DeviceID->Split(TEXT("@"), nullptr, &DeviceName))
+	{
+		DeviceName = *DeviceID;
+	}
+
+	return FText::FromString(DeviceName);
+}
+
diff --git a/Engine/Plugins/Developer/ProjectLauncher/Source/ProjectLauncher/Public/Widgets/Shared/SCustomLaunchDeviceCombo.h b/Engine/Plugins/Developer/ProjectLauncher/Source/ProjectLauncher/Public/Widgets/Shared/SCustomLaunchDeviceCombo.h
--- a/Engine/Plugins/Developer/ProjectLauncher/Source/ProjectLauncher/Public/Widgets/Shared/SCustomLaunchDeviceCombo.h
+++ b/Engine/Plugins/Developer/ProjectLauncher/Source/ProjectLauncher/Public/Widgets/Shared/SCustomLaunchDeviceCombo.h
@@ -16,7 +16,6 @@ public:
 	SLATE_BEGIN_ARGS(SCustomLaunchDeviceCombo)
 		: _AllPlatforms(false)
 		{}
-		SLATE_EVENT(FOnDeviceRemoved, OnDeviceRemoved)
 		SLATE_EVENT(FOnSelectionChanged, OnSelectionChanged);
 		SLATE_ATTRIBUTE(TArray<FString>, SelectedDevices);
 		SLATE_ATTRIBUTE(TArray<FString>, Platforms);
@@ -28,13 +27,13 @@ public:
 
 protected:
 
-	UE_API TSharedRef<SWidget> GenerateDeviceProxyListWidget(TSharedPtr<ITargetDeviceProxy> DeviceProxy) const;
-	UE_API void OnDeviceProxySelectionChanged(TSharedPtr<ITargetDeviceProxy> DeviceProxy, ESelectInfo::Type InSelectInfo);
-	UE_API const FSlateBrush* GetSelectedDeviceProxyBrush() const;
-	UE_API FText GetSelectedDeviceProxyName() const;
+	UE_API TSharedRef<SWidget> GenerateDeviceListWidget(TSharedPtr<FString> DeviceID) const;
+	UE_API void OnDeviceSelectionChanged(TSharedPtr<FString> DeviceID, ESelectInfo::Type InSelectInfo);
+	UE_API const FSlateBrush* GetSelectedDeviceBrush() const;
+	UE_API FText GetSelectedDeviceName() const;
 
 private:
-	TSharedPtr<SComboBox<TSharedPtr<ITargetDeviceProxy>> > DeviceProxyComboBox;
+	TSharedPtr<SComboBox<TSharedPtr<FString>> > DeviceComboBox;
 
 };
 
diff --git a/Engine/Plugins/Developer/ProjectLauncher/Source/ProjectLauncher/Public/Widgets/Shared/SCustomLaunchDeviceListView.h b/Engine/Plugins/Developer/ProjectLauncher/Source/ProjectLauncher/Public/Widgets/Shared/SCustomLaunchDeviceListView.h
--- a/Engine/Plugins/Developer/ProjectLauncher/Source/ProjectLauncher/Public/Widgets/Shared/SCustomLaunchDeviceListView.h
+++ b/Engine/Plugins/Developer/ProjectLauncher/Source/ProjectLauncher/Public/Widgets/Shared/SCustomLaunchDeviceListView.h
@@ -19,7 +19,6 @@ public:
 		: _AllPlatforms(false)
 		, _SingleSelect(false)
 		{}
-		SLATE_EVENT(FOnDeviceRemoved, OnDeviceRemoved)
 		SLATE_EVENT(FOnSelectionChanged, OnSelectionChanged);
 		SLATE_ATTRIBUTE(TArray<FString>, SelectedDevices);
 		SLATE_ATTRIBUTE(TArray<FString>, Platforms);
@@ -33,13 +32,13 @@ public:
 protected:
 	bool bSingleSelect;
 
-	UE_API TSharedRef<ITableRow> GenerateDeviceProxyRow(TSharedPtr<ITargetDeviceProxy> DeviceProxy, const TSharedRef<STableViewBase>& OwnerTable);
-	UE_API ECheckBoxState IsDeviceProxyChecked(TSharedPtr<ITargetDeviceProxy> DeviceProxy) const;
-	UE_API void OnDeviceProxyCheckStateChanged(ECheckBoxState NewState, TSharedPtr<ITargetDeviceProxy> DeviceProxy);
+	UE_API TSharedRef<ITableRow> GenerateDeviceRow(TSharedPtr<FString> DeviceID, const TSharedRef<STableViewBase>& OwnerTable);
+	UE_API ECheckBoxState IsDeviceChecked(TSharedPtr<FString> DeviceID) const;
+	UE_API void OnDeviceCheckStateChanged(ECheckBoxState NewState, TSharedPtr<FString> DeviceID);
 	UE_API virtual void OnDeviceListRefreshed() override;
 
 private:
-	TSharedPtr<SListView<TSharedPtr<ITargetDeviceProxy>> > DeviceProxyListView;
+	TSharedPtr<SListView<TSharedPtr<FString>>> DeviceListView;
 
 };
 
diff --git a/Engine/Plugins/Developer/ProjectLauncher/Source/ProjectLauncher/Public/Widgets/Shared/SCustomLaunchDeviceWidgetBase.h b/Engine/Plugins/Developer/ProjectLauncher/Source/ProjectLauncher/Public/Widgets/Shared/SCustomLaunchDeviceWidgetBase.h
--- a/Engine/Plugins/Developer/ProjectLauncher/Source/ProjectLauncher/Public/Widgets/Shared/SCustomLaunchDeviceWidgetBase.h
+++ b/Engine/Plugins/Developer/ProjectLauncher/Source/ProjectLauncher/Public/Widgets/Shared/SCustomLaunchDeviceWidgetBase.h
@@ -10,12 +10,16 @@
 class ITargetDeviceProxy;
 class ITargetDeviceProxyManager;
 
+namespace PlatformInfo
+{
+	struct FTargetPlatformInfo;
+}
+
 class SCustomLaunchDeviceWidgetBase
 	: public SCompoundWidget
 {
 public:
 	DECLARE_DELEGATE_OneParam(FOnSelectionChanged, TArray<FString> );
-	DECLARE_DELEGATE_OneParam(FOnDeviceRemoved, FString );
 
 public:
 	UE_API void Construct();
@@ -30,7 +34,6 @@ protected:
 	TAttribute<TArray<FString>> Platforms;
 	TAttribute<TArray<FString>> SelectedDevices;
 	FOnSelectionChanged OnSelectionChanged;
-	FOnDeviceRemoved OnDeviceRemoved;
 	bool bAllPlatforms = false;
 
 	UE_API void OnDeviceProxyAdded(const TSharedRef<ITargetDeviceProxy>& DeviceProxy);
@@ -38,7 +41,12 @@ protected:
 
 	UE_API const TSharedRef<ITargetDeviceProxyManager> GetDeviceProxyManager() const;
 
-	TArray<TSharedPtr<ITargetDeviceProxy>> DeviceProxyList;
+	UE_API const PlatformInfo::FTargetPlatformInfo* GetPlatformInfoForDeviceID( const FString& DeviceID ) const;
+	UE_API FText GetDisplayNameForDevice( const TSharedPtr<FString>& DeviceID ) const;
+
+	TSet<FString> KnownDeviceIDs;
+	TArray<TSharedPtr<FString>> AllDeviceIDs;
+
 };
 
 #undef UE_API
