bug fix for not being able to rename tasks in the basic launch profile

#rnx
#Synced-CL 50857758

diff --git a/Engine/Plugins/Developer/ProjectLauncher/Source/ProjectLauncher/Private/Widgets/Profiles/SCustomLaunchCustomProfileEditor.cpp b/Engine/Plugins/Developer/ProjectLauncher/Source/ProjectLauncher/Private/Widgets/Profiles/SCustomLaunchCustomProfileEditor.cpp
--- a/Engine/Plugins/Developer/ProjectLauncher/Source/ProjectLauncher/Private/Widgets/Profiles/SCustomLaunchCustomProfileEditor.cpp
+++ b/Engine/Plugins/Developer/ProjectLauncher/Source/ProjectLauncher/Private/Widgets/Profiles/SCustomLaunchCustomProfileEditor.cpp
@@ -1294,7 +1294,7 @@ void SCustomLaunchCustomProfileEditor::MakeUATCommandSubMenu( FMenuBuilder& Menu
 
 void SCustomLaunchCustomProfileEditor::MakeTreeNodeSubMenu( FMenuBuilder& MenuBuilder, ProjectLauncher::FLaunchProfileTreeNodePtr TreeNode)
 {
-	if (TreeNode->UATCommand.IsValid() && !Model->CanUseSimplifiedLayout(CurrentProfile.ToSharedRef()))
+	if (TreeNode->UATCommand.IsValid() && (!Model->CanUseSimplifiedLayout(CurrentProfile.ToSharedRef()) || UATCommandNameFields.Contains(TreeNode->UATCommand)))
 	{
 		MakeUATCommandSubMenu(MenuBuilder, TreeNode->UATCommand.ToSharedRef());
 	}
