added config helper functions to build cook run command extensions to allow configuration data to be stored per-build cook run.

#rnx

#Synced-CL 50857758

diff --git a/Engine/Plugins/Developer/ProjectLauncher/Source/ProjectLauncher/Private/Extension/BuildCookRunCommandExtension.cpp b/Engine/Plugins/Developer/ProjectLauncher/Source/ProjectLauncher/Private/Extension/BuildCookRunCommandExtension.cpp
--- a/Engine/Plugins/Developer/ProjectLauncher/Source/ProjectLauncher/Private/Extension/BuildCookRunCommandExtension.cpp
+++ b/Engine/Plugins/Developer/ProjectLauncher/Source/ProjectLauncher/Private/Extension/BuildCookRunCommandExtension.cpp
@@ -41,6 +41,59 @@ namespace ProjectLauncher
 		return ProfileTreeData.AddSubHeading(Extension->GetInternalName(), Extension->GetDisplayName(), INT32_MAX );
 	}
 
+	FString FBuildCookRunExtension::GetConfigString(EConfig Config, const TCHAR* Name, const TCHAR* DefaultValue) const
+	{
+		return ExtensionInstance->GetConfigString(Config, *GetConfigKeyName(Name), DefaultValue);
+	}
+
+	bool FBuildCookRunExtension::GetConfigBool(EConfig Config, const TCHAR* Name, bool DefaultValue) const
+	{
+		return ExtensionInstance->GetConfigBool(Config, *GetConfigKeyName(Name), DefaultValue);
+	}
+
+	int32 FBuildCookRunExtension::GetConfigInteger(EConfig Config, const TCHAR* Name, int32 DefaultValue) const
+	{
+		return ExtensionInstance->GetConfigInteger(Config, *GetConfigKeyName(Name), DefaultValue);
+	}
+
+	float FBuildCookRunExtension::GetConfigFloat(EConfig Config, const TCHAR* Name, float DefaultValue) const
+	{
+		return ExtensionInstance->GetConfigFloat(Config, *GetConfigKeyName(Name), DefaultValue);
+	}
+
+	void FBuildCookRunExtension::SetConfigString(EConfig Config, const TCHAR* Name, const FString& Value) const
+	{
+		ExtensionInstance->SetConfigString(Config, *GetConfigKeyName(Name), Value);
+	}
+
+	void FBuildCookRunExtension::SetConfigBool(EConfig Config, const TCHAR* Name, const bool& Value) const
+	{
+		ExtensionInstance->SetConfigBool(Config, *GetConfigKeyName(Name), Value);
+	}
+
+	void FBuildCookRunExtension::SetConfigInteger(EConfig Config, const TCHAR* Name, const int32& Value) const
+	{
+		ExtensionInstance->SetConfigInteger(Config, *GetConfigKeyName(Name), Value);
+	}
+
+	void FBuildCookRunExtension::SetConfigFloat(EConfig Config, const TCHAR* Name, const float& Value) const
+	{
+		ExtensionInstance->SetConfigFloat(Config, *GetConfigKeyName(Name), Value);
+	}
+
+	FString FBuildCookRunExtension::GetConfigKeyName(const TCHAR* Name) const
+	{
+		return FString::Printf(TEXT("%s.%s"), BuildCookRun->GetInternalName(), Name);
+	}
+
+
+
+
+
+
+
+
+
 
 
 
diff --git a/Engine/Plugins/Developer/ProjectLauncher/Source/ProjectLauncher/Public/Extension/BuildCookRunCommandExtension.h b/Engine/Plugins/Developer/ProjectLauncher/Source/ProjectLauncher/Public/Extension/BuildCookRunCommandExtension.h
--- a/Engine/Plugins/Developer/ProjectLauncher/Source/ProjectLauncher/Public/Extension/BuildCookRunCommandExtension.h
+++ b/Engine/Plugins/Developer/ProjectLauncher/Source/ProjectLauncher/Public/Extension/BuildCookRunCommandExtension.h
@@ -14,6 +14,8 @@ namespace ProjectLauncher
 	class FBuildCookRunExtension : public TSharedFromThis<FBuildCookRunExtension>
 	{
 	public:
+		typedef FLaunchExtensionInstance::EConfig EConfig;
+
 		struct FArgs
 		{
 			ILauncherProfileRef Profile;
@@ -92,8 +94,86 @@ namespace ProjectLauncher
 		 */
 		inline TSharedRef<FModel> GetModel() const { return Model; }
 
+
+		/**
+		 * Read a configuration string value
+		 *
+		 * @param Config where the value is stored
+		 * @param Name name of the value to read
+		 * @param DefaultValue value to return if the value isn't found
+		 * @returns the value of the property, or DefaultValue if it isn't found
+		 */
+		UE_API FString GetConfigString(EConfig Config, const TCHAR* Name, const TCHAR* DefaultValue = TEXT("")) const;
+
+		/**
+		 * Read a configuration bool value
+		 *
+		 * @param Config where the value is stored
+		 * @param Name name of the value to read
+		 * @param DefaultValue value to return if the value isn't found
+		 * @returns the value of the property, or DefaultValue if it isn't found
+		 */
+		UE_API bool GetConfigBool(EConfig Config, const TCHAR* Name, bool DefaultValue = false) const;
+
+		/**
+		 * Read a configuration int value
+		 *
+		 * @param Config where the value is stored
+		 * @param Name name of the value to read
+		 * @param DefaultValue value to return if the value isn't found
+		 * @returns the value of the property, or DefaultValue if it isn't found
+		 */
+		UE_API int32 GetConfigInteger(EConfig Config, const TCHAR* Name, int32 DefaultValue = 0) const;
+
+		/**
+		 * Read a configuration float value
+		 *
+		 * @param Config where the value is stored
+		 * @param Name name of the value to read
+		 * @param DefaultValue value to return if the value isn't found
+		 * @returns the value of the property, or DefaultValue if it isn't found
+		 */
+		UE_API float GetConfigFloat(EConfig Config, const TCHAR* Name, float DefaultValue = 0.0f) const;
+
+		/**
+		 * Write a configuration string value
+		 *
+		 * @param Config where the value is stored
+		 * @param Name name of the value to write
+		 * @param Value value to write
+		 */
+		UE_API void SetConfigString(EConfig Config, const TCHAR* Name, const FString& Value) const;
+
+		/**
+		 * Write a configuration bool value
+		 *
+		 * @param Config where the value is stored
+		 * @param Name name of the value to write
+		 * @param Value value to write
+		 */
+		UE_API void SetConfigBool(EConfig Config, const TCHAR* Name, const bool& Value) const;
+
+		/**
+		 * Write a configuration int value
+		 *
+		 * @param Config where the value is stored
+		 * @param Name name of the value to write
+		 * @param Value value to write
+		 */
+		UE_API void SetConfigInteger(EConfig Config, const TCHAR* Name, const int32& Value) const;
+
+		/**
+		 * Write a configuration float value
+		 *
+		 * @param Config where the value is stored
+		 * @param Name name of the value to write
+		 * @param Value value to write
+		 */
+		UE_API void SetConfigFloat(EConfig Config, const TCHAR* Name, const float& Value) const;
+
 	private:
 		void OnUATCommandValidation( ILauncherProfileUATCommandRef UATCommand );
+		UE_API FString GetConfigKeyName(const TCHAR* Name) const;
 
 		ILauncherProfileRef Profile;
 		ILauncherProfileBuildCookRunRef BuildCookRun;
