Add two new launch extensions
- add custom parameters to a specific BuildCookRun's UAT command
- add a fully user-defined UAT command task

also includes a fix to prevent the command line field from showing the + button until there are UAT command line extensions

#rnx

#Synced-CL 50857758

diff --git a/Engine/Plugins/Developer/ProjectLauncher/Source/CommonLaunchExtensions/Private/CommonLaunchExtensionsModule.cpp b/Engine/Plugins/Developer/ProjectLauncher/Source/CommonLaunchExtensions/Private/CommonLaunchExtensionsModule.cpp
--- a/Engine/Plugins/Developer/ProjectLauncher/Source/CommonLaunchExtensions/Private/CommonLaunchExtensionsModule.cpp
+++ b/Engine/Plugins/Developer/ProjectLauncher/Source/CommonLaunchExtensions/Private/CommonLaunchExtensionsModule.cpp
@@ -11,6 +11,8 @@
 #include "DeprecatedProperties/DeprecatedPropertiesLaunchExtension.h"
 #include "BuildSync/BuildSyncLaunchExtension.h"
 #include "UgsSync/UgsSyncLaunchExtension.h"
+#include "UserUATArgs/UserUATArgsLaunchExtension.h"
+#include "UserUATCommand/UserUATCommandLaunchExtension.h"
 
 class FCommonLaunchExtensionsModule : public IModuleInterface
 {
@@ -26,6 +28,8 @@ public:
 		Extensions.Add( MakeShared<FDeprecatedPropertiesLaunchExtension>() );
 		Extensions.Add( MakeShared<FBuildSyncLaunchExtension>() );
 		Extensions.Add( MakeShared<FUGSSyncLaunchExtension>() );
+		Extensions.Add( MakeShared<FUserUATArgsLaunchExtension>() );
+		Extensions.Add( MakeShared<FUserUATCommandLaunchExtension>() );
 
 		for (TSharedRef<ProjectLauncher::FLaunchExtension> Extension : Extensions)
 		{
diff --git a/Engine/Plugins/Developer/ProjectLauncher/Source/CommonLaunchExtensions/Private/UserUATArgs/UserUATArgsLaunchExtension.cpp b/Engine/Plugins/Developer/ProjectLauncher/Source/CommonLaunchExtensions/Private/UserUATArgs/UserUATArgsLaunchExtension.cpp
new file mode 100644
--- /dev/null
+++ b/Engine/Plugins/Developer/ProjectLauncher/Source/CommonLaunchExtensions/Private/UserUATArgs/UserUATArgsLaunchExtension.cpp
@@ -0,0 +1,78 @@
+// Copyright Epic Games, Inc. All Rights Reserved.
+
+#include "UserUATArgs/UserUATArgsLaunchExtension.h"
+#include "Model/ProjectLauncherModel.h"
+
+#define LOCTEXT_NAMESPACE "FUserUATArgsLaunchExtensionInstance"
+
+
+
+FUserUATArgsLaunchExtensionInstance::FUserUATArgsLaunchExtensionInstance( FArgs& InArgs )
+	: ProjectLauncher::FBuildCookRunCommandExtensionInstance(InArgs)
+	, Owner(StaticCastSharedRef<FUserUATArgsLaunchExtension>(InArgs.Extension))
+{
+}
+
+TSharedRef<ProjectLauncher::FBuildCookRunExtension> FUserUATArgsLaunchExtensionInstance::CreateBuildCookRunExtension( const ProjectLauncher::FBuildCookRunExtension::FArgs& InArgs )
+{
+	return MakeShared<FBuildCookRunInstance>(InArgs);
+}
+
+
+
+
+FUserUATArgsLaunchExtensionInstance::FBuildCookRunInstance::FBuildCookRunInstance( const FArgs& InArgs )
+	: ProjectLauncher::FBuildCookRunExtension( InArgs )
+{
+}
+
+
+void FUserUATArgsLaunchExtensionInstance::FBuildCookRunInstance::CustomizeTree( ProjectLauncher::FLaunchProfileTreeNode& ProfileTreeNode )
+{
+	AddDefaultHeading(ProfileTreeNode)
+		.AddCommandLineString( LOCTEXT("ExtraParamsLabel", "Additional UAT Parameters"),
+			{
+				.GetValue = [this]()              { return GetConfigString(EConfig::PerProfile, TEXT("ExtraUATParams")); },
+				.SetValue = [this](FString Value) { SetConfigString(EConfig::PerProfile, TEXT("ExtraUATParams"), Value); },
+				.GetDefaultValue = [this]()       { return FString(); },
+			},
+			ProjectLauncher::CmdLineType::UAT
+		);
+	;
+}
+
+void FUserUATArgsLaunchExtensionInstance::FBuildCookRunInstance::CustomizeUATCommandLine(FString& InOutCommandLine)
+{
+	FString ExtraUATParams = GetConfigString(EConfig::PerProfile, TEXT("ExtraUATParams")).TrimStartAndEnd();
+	if (!ExtraUATParams.IsEmpty())
+	{
+		InOutCommandLine += TEXT(" ");
+		InOutCommandLine += ExtraUATParams;
+	}
+}
+
+
+
+
+
+
+
+
+TSharedPtr<ProjectLauncher::FLaunchExtensionInstance> FUserUATArgsLaunchExtension::CreateInstanceForProfile( ProjectLauncher::FLaunchExtensionInstance::FArgs& InArgs )
+{
+	return MakeShared<FUserUATArgsLaunchExtensionInstance>(InArgs);
+}
+
+const TCHAR* FUserUATArgsLaunchExtension::GetInternalName() const
+{
+	return TEXT("UserUATArgs");
+}
+
+FText FUserUATArgsLaunchExtension::GetDisplayName() const
+{
+	return LOCTEXT("ExtensionName", "Custom UAT Parameters");
+}
+
+
+
+#undef LOCTEXT_NAMESPACE
diff --git a/Engine/Plugins/Developer/ProjectLauncher/Source/CommonLaunchExtensions/Private/UserUATArgs/UserUATArgsLaunchExtension.h b/Engine/Plugins/Developer/ProjectLauncher/Source/CommonLaunchExtensions/Private/UserUATArgs/UserUATArgsLaunchExtension.h
new file mode 100644
--- /dev/null
+++ b/Engine/Plugins/Developer/ProjectLauncher/Source/CommonLaunchExtensions/Private/UserUATArgs/UserUATArgsLaunchExtension.h
@@ -0,0 +1,36 @@
+// Copyright Epic Games, Inc. All Rights Reserved.
+
+#pragma once
+
+#include "Extension/BuildCookRunCommandExtension.h"
+
+
+class FUserUATArgsLaunchExtensionInstance : public ProjectLauncher::FBuildCookRunCommandExtensionInstance
+{
+public:
+	FUserUATArgsLaunchExtensionInstance( FArgs& InArgs );
+	virtual ~FUserUATArgsLaunchExtensionInstance() = default;
+
+	virtual TSharedRef<ProjectLauncher::FBuildCookRunExtension> CreateBuildCookRunExtension( const ProjectLauncher::FBuildCookRunExtension::FArgs& InArgs ) override;
+
+protected:
+	TSharedRef<class FUserUATArgsLaunchExtension> Owner;
+
+	class FBuildCookRunInstance : public ProjectLauncher::FBuildCookRunExtension
+	{
+	public:
+		FBuildCookRunInstance(const ProjectLauncher::FBuildCookRunExtension::FArgs& InArgs );
+		virtual void CustomizeTree( ProjectLauncher::FLaunchProfileTreeNode& ProfileTreeNode ) override;
+		virtual void CustomizeUATCommandLine(FString& InOutCommandLine) override;
+
+	};
+};
+
+
+class FUserUATArgsLaunchExtension : public ProjectLauncher::FBuildCookRunCommandExtension
+{
+public:
+	virtual TSharedPtr<ProjectLauncher::FLaunchExtensionInstance> CreateInstanceForProfile( ProjectLauncher::FLaunchExtensionInstance::FArgs& InArgs ) override;
+	virtual const TCHAR* GetInternalName() const override;
+	virtual FText GetDisplayName() const override;
+};
\ No newline at end of file
diff --git a/Engine/Plugins/Developer/ProjectLauncher/Source/CommonLaunchExtensions/Private/UserUATCommand/UserUATCommandLaunchExtension.cpp b/Engine/Plugins/Developer/ProjectLauncher/Source/CommonLaunchExtensions/Private/UserUATCommand/UserUATCommandLaunchExtension.cpp
new file mode 100644
--- /dev/null
+++ b/Engine/Plugins/Developer/ProjectLauncher/Source/CommonLaunchExtensions/Private/UserUATCommand/UserUATCommandLaunchExtension.cpp
@@ -0,0 +1,80 @@
+// Copyright Epic Games, Inc. All Rights Reserved.
+
+#include "UserUATCommand/UserUATCommandLaunchExtension.h"
+#include "Model/ProjectLauncherModel.h"
+
+
+#define LOCTEXT_NAMESPACE "FUserUATCommandLaunchExtensionInstance"
+
+FUserUATCommandLaunchExtensionInstance::FUserUATCommandLaunchExtensionInstance(FArgs& InArgs)
+	: ProjectLauncher::FCustomUATCommandLaunchExtensionInstance(InArgs)
+{
+}
+
+
+void FUserUATCommandLaunchExtensionInstance::CustomizeTree(ProjectLauncher::FLaunchProfileTreeData& ProfileTreeData)
+{
+	FLaunchExtensionInstance::CustomizeTree(ProfileTreeData);
+
+	AddDefaultHeading(ProfileTreeData)
+		.AddString(LOCTEXT("UATCommandLabel", "UAT Command"),
+			{
+				.GetValue = [this]()              { return GetUATCommand()->GetUATCommand(); },
+				.SetValue = [this](FString Value) { GetUATCommand()->SetUATCommand(*Value); GetProfile()->RefreshCustomWarningsAndErrors(GetUATCommand()); },
+				.Validation = ProjectLauncher::FValidation({TEXT("Validation_UserUATCmd_NoCommand")}),
+			}
+		)
+		.AddCommandLineString(LOCTEXT("ExtraParamsLabel", "UAT Parameters"),
+			{
+				.GetValue = [this]()              { return GetConfigString(EConfig::PerProfile, TEXT("UATParams")); },
+				.SetValue = [this](FString Value) { SetConfigString(EConfig::PerProfile, TEXT("UATParams"), Value); },
+				.GetDefaultValue = [this]()       { return FString(); },
+			},
+			ProjectLauncher::CmdLineType::UAT
+		);
+	;
+}
+
+void FUserUATCommandLaunchExtensionInstance::CustomizeUATCommandLine(FString& InOutCommandLine)
+{
+	FString ExtraUATParams = GetConfigString(EConfig::PerProfile, TEXT("UATParams")).TrimStartAndEnd();
+	if (!ExtraUATParams.IsEmpty())
+	{
+		InOutCommandLine += TEXT(" ");
+		InOutCommandLine += ExtraUATParams;
+	}
+}
+
+void FUserUATCommandLaunchExtensionInstance::OnValidateProfile()
+{
+	FString UATCommandName = GetUATCommand()->GetUATCommand().TrimStartAndEnd();
+	if (UATCommandName.IsEmpty() || UATCommandName.Contains(TEXT(" ")) || UATCommandName.Contains("-") )
+	{
+		GetProfile()->AddCustomError(TEXT("Validation_UserUATCmd_NoCommand"), LOCTEXT("Validation_UserUATCmd_NoCommand", "You must specify a valid UAT command"), GetUATCommand());
+	}
+}
+
+void FUserUATCommandLaunchExtensionInstance::OnUATCommandAdded(ILauncherProfileUATCommandRef InUATCommand)
+{
+	InUATCommand->SetUATCommand(TEXT(""));
+}
+
+
+
+
+TSharedPtr<ProjectLauncher::FLaunchExtensionInstance> FUserUATCommandLaunchExtension::CreateInstanceForProfile(ProjectLauncher::FLaunchExtensionInstance::FArgs& InArgs)
+{
+	return MakeShared<FUserUATCommandLaunchExtensionInstance>(InArgs);
+}
+
+const TCHAR* FUserUATCommandLaunchExtension::GetInternalName() const
+{
+	return TEXT("UserUATCommand");
+}
+
+FText FUserUATCommandLaunchExtension::GetDisplayName() const
+{
+	return LOCTEXT("ExtensionName", "Custom UAT Command");
+}
+
+#undef LOCTEXT_NAMESPACE
diff --git a/Engine/Plugins/Developer/ProjectLauncher/Source/CommonLaunchExtensions/Private/UserUATCommand/UserUATCommandLaunchExtension.h b/Engine/Plugins/Developer/ProjectLauncher/Source/CommonLaunchExtensions/Private/UserUATCommand/UserUATCommandLaunchExtension.h
new file mode 100644
--- /dev/null
+++ b/Engine/Plugins/Developer/ProjectLauncher/Source/CommonLaunchExtensions/Private/UserUATCommand/UserUATCommandLaunchExtension.h
@@ -0,0 +1,26 @@
+// Copyright Epic Games, Inc. All Rights Reserved.
+
+#pragma once
+
+#include "Extension/CustomUATCommandLaunchExtension.h"
+
+class FUserUATCommandLaunchExtensionInstance : public ProjectLauncher::FCustomUATCommandLaunchExtensionInstance
+{
+public:
+	FUserUATCommandLaunchExtensionInstance(FArgs& InArgs);
+
+	virtual void CustomizeTree(ProjectLauncher::FLaunchProfileTreeData& ProfileTreeData) override;
+	virtual void CustomizeUATCommandLine(FString& InOutCommandLine) override;
+	
+protected:
+	virtual void OnUATCommandAdded(ILauncherProfileUATCommandRef InUATCommand) override;
+	virtual void OnValidateProfile() override;
+};
+
+class FUserUATCommandLaunchExtension: public ProjectLauncher::FCustomUATCommandLaunchExtension
+{
+public:
+	virtual TSharedPtr<ProjectLauncher::FLaunchExtensionInstance> CreateInstanceForProfile(ProjectLauncher::FLaunchExtensionInstance::FArgs& InArgs) override;
+	virtual const TCHAR* GetInternalName() const override;
+	virtual FText GetDisplayName() const override;
+};
\ No newline at end of file
diff --git a/Engine/Plugins/Developer/ProjectLauncher/Source/ProjectLauncher/Private/ProfileTree/LaunchProfileTreeData.cpp b/Engine/Plugins/Developer/ProjectLauncher/Source/ProjectLauncher/Private/ProfileTree/LaunchProfileTreeData.cpp
--- a/Engine/Plugins/Developer/ProjectLauncher/Source/ProjectLauncher/Private/ProfileTree/LaunchProfileTreeData.cpp
+++ b/Engine/Plugins/Developer/ProjectLauncher/Source/ProjectLauncher/Private/ProfileTree/LaunchProfileTreeData.cpp
@@ -784,6 +784,21 @@ namespace ProjectLauncher
 			return MenuBuilder.MakeWidget();
 		};
 
+		auto GetExtensionButtonVisibility = [this, TypeName = FString(InType)]()
+		{
+			// see if there's an extension that can support our command line type
+			for (TSharedPtr<FLaunchExtensionInstance> ExtensionInstance : TreeData->ExtensionInstances)
+			{
+				ProjectLauncher::ICmdLineParametersExtensionFactory* CmdLineParametersFactory = ExtensionInstance->AsCmdLineParametersFactory();
+				if (CmdLineParametersFactory != nullptr && CmdLineParametersFactory->GetCmdLineType() == TypeName)
+				{
+					return EVisibility::Visible;
+				}
+			}
+
+			return EVisibility::Collapsed;
+		};
+
 		return AddWidget( InName,
 			MoveTemp(WidgetCallbacks),
 			SNew(SHorizontalBox)
@@ -811,7 +826,7 @@ namespace ProjectLauncher
 				.ComboButtonStyle(FAppStyle::Get(), "SimpleComboButton")
 				.ToolTipText(LOCTEXT("CmdLineOptionsLabel", "Add a special parameter to command line"))
 				.OnGetMenuContent_Lambda(OnGetCmdlineParameterMenuContent)
-				.Visibility_Lambda( [this]() { return TreeData->ExtensionInstances.Num() > 0 ? EVisibility::Visible : EVisibility::Collapsed; } )
+				.Visibility_Lambda(GetExtensionButtonVisibility)
 				.MenuPlacement(MenuPlacement_BelowRightAnchor)
 				.HasDownArrow(false)
 				.ButtonContent()
