expose 'OnMenuOpenChanged' in project launcher common combo box

#rnx

#Synced-CL 50857758

diff --git a/Engine/Plugins/Developer/ProjectLauncher/Source/ProjectLauncher/Public/Widgets/Shared/SCustomLaunchCombo.h b/Engine/Plugins/Developer/ProjectLauncher/Source/ProjectLauncher/Public/Widgets/Shared/SCustomLaunchCombo.h
--- a/Engine/Plugins/Developer/ProjectLauncher/Source/ProjectLauncher/Public/Widgets/Shared/SCustomLaunchCombo.h
+++ b/Engine/Plugins/Developer/ProjectLauncher/Source/ProjectLauncher/Public/Widgets/Shared/SCustomLaunchCombo.h
@@ -34,6 +34,7 @@ public:
 		SLATE_EVENT(FIsItemEnabled, IsItemEnabled)
 		SLATE_EVENT(FGetItemIcon, GetItemIcon)
 		SLATE_EVENT(FOnSelectedItemsChanged, OnSelectedItemsChanged)
+		SLATE_EVENT(FOnIsOpenChanged, OnMenuOpenChanged)
 		SLATE_ATTRIBUTE(TArray<T>, Items)
 		SLATE_ATTRIBUTE(T, SelectedItem)
 		SLATE_ATTRIBUTE(TArray<T>, SelectedItems)
@@ -77,6 +78,7 @@ public:
 				.Font(FCoreStyle::Get().GetFontStyle("SmallFont"))
 			]
 			.OnGetMenuContent(this, &SCustomLaunchCombo<T>::MakeWidget)
+			.OnMenuOpenChanged(InArgs._OnMenuOpenChanged)
 		];
 	}
 
