only the editable field is disabled for normal property fields, not the whole row

#rnx

#Synced-CL 50857758

#Synced-CL 50857758

diff --git a/Engine/Plugins/Developer/ProjectLauncher/Source/ProjectLauncher/Private/ProfileTree/LaunchProfileTreeData.cpp b/Engine/Plugins/Developer/ProjectLauncher/Source/ProjectLauncher/Private/ProfileTree/LaunchProfileTreeData.cpp
--- a/Engine/Plugins/Developer/ProjectLauncher/Source/ProjectLauncher/Private/ProfileTree/LaunchProfileTreeData.cpp
+++ b/Engine/Plugins/Developer/ProjectLauncher/Source/ProjectLauncher/Private/ProfileTree/LaunchProfileTreeData.cpp
@@ -398,7 +398,6 @@ namespace ProjectLauncher
 		FCallbacks WidgetCallbacks
 		{
 			.IsVisible = BooleanCallbacks.IsVisible,
-			.IsEnabled = BooleanCallbacks.IsEnabled,
 			.Validation = BooleanCallbacks.Validation,
 		};
 		if (BooleanCallbacks.GetDefaultValue != nullptr)
@@ -426,11 +425,17 @@ namespace ProjectLauncher
 			return BooleanCallbacks.GetValue() ? ECheckBoxState::Checked : ECheckBoxState::Unchecked;
 		};
 
+		auto IsEnabled = [BooleanCallbacks]()
+		{
+			return !BooleanCallbacks.IsEnabled.IsSet() || BooleanCallbacks.IsEnabled();
+		};
+
 		return AddWidget( InName,
 			MoveTemp(WidgetCallbacks),
 			SNew(SCheckBox)
 			.OnCheckStateChanged_Lambda(SetCheckState)
 			.IsChecked_Lambda(GetCheckState)
+			.IsEnabled_Lambda(IsEnabled)
 		);
 	}
 
@@ -442,7 +447,6 @@ namespace ProjectLauncher
 		FCallbacks WidgetCallbacks
 		{
 			.IsVisible = StringCallbacks.IsVisible,
-			.IsEnabled = StringCallbacks.IsEnabled,
 			.Validation = StringCallbacks.Validation,
 		};
 		if (StringCallbacks.GetDefaultValue != nullptr)
@@ -470,6 +474,12 @@ namespace ProjectLauncher
 			return FText::FromString(StringCallbacks.GetValue());
 		};
 
+		auto IsReadOnly = [StringCallbacks]()
+		{
+			return StringCallbacks.IsEnabled && !StringCallbacks.IsEnabled();
+		};
+
+
 		return AddWidget( InName,
 			MoveTemp(WidgetCallbacks),
 			SNew(SMultiLineEditableTextBox)
@@ -478,6 +488,7 @@ namespace ProjectLauncher
 			.Font(FCoreStyle::Get().GetFontStyle("SmallFont"))
 			.Text_Lambda( GetString )
 			.OnTextCommitted_Lambda( SetString )
+			.IsReadOnly_Lambda( IsReadOnly )
 		);
 
 	}
@@ -490,7 +501,6 @@ namespace ProjectLauncher
 		FCallbacks WidgetCallbacks
 		{
 			.IsVisible = StringCallbacks.IsVisible,
-			.IsEnabled = StringCallbacks.IsEnabled,
 			.Validation = StringCallbacks.Validation,
 		};
 		if (StringCallbacks.GetDefaultValue != nullptr)
@@ -546,6 +556,10 @@ namespace ProjectLauncher
 			return FReply::Handled();
 		};
 
+		auto IsEnabled = [StringCallbacks]()
+		{
+			return !StringCallbacks.IsEnabled.IsSet() || StringCallbacks.IsEnabled();
+		};
 
 		return AddWidget( InName,
 			MoveTemp(WidgetCallbacks),
@@ -560,6 +574,7 @@ namespace ProjectLauncher
 				.Font(FCoreStyle::Get().GetFontStyle("SmallFont"))
 				.Text_Lambda( GetString )
 				.OnTextCommitted_Lambda( SetString )
+				.IsEnabled_Lambda( IsEnabled )
 			]
 
 			// browse button
@@ -572,6 +587,7 @@ namespace ProjectLauncher
 				.ButtonStyle(FProjectLauncherStyle::Get(), "HoverHintOnly")
 				.ToolTipText(LOCTEXT("DirBrowseTip", "Browse for a folder"))
 				.OnClicked_Lambda(OnBrowse)
+				.IsEnabled_Lambda(IsEnabled)
 				.ContentPadding(2)
 				[
 					SNew(SImage)
@@ -594,7 +610,6 @@ namespace ProjectLauncher
 		FCallbacks WidgetCallbacks
 		{
 			.IsVisible = StringCallbacks.IsVisible,
-			.IsEnabled = StringCallbacks.IsEnabled,
 		};
 
 		if (StringCallbacks.GetDefaultValue != nullptr)
@@ -656,6 +671,10 @@ namespace ProjectLauncher
 			return FReply::Handled();
 		};
 
+		auto IsEnabled = [StringCallbacks]()
+		{
+			return !StringCallbacks.IsEnabled.IsSet() || StringCallbacks.IsEnabled();
+		};
 
 		return AddWidget( InName,
 			MoveTemp(WidgetCallbacks),
@@ -670,6 +689,7 @@ namespace ProjectLauncher
 				.Font(FCoreStyle::Get().GetFontStyle("SmallFont"))
 				.Text_Lambda( GetString )
 				.OnTextCommitted_Lambda( SetString )
+				.IsEnabled_Lambda( IsEnabled )
 			]
 
 			// browse button
@@ -682,6 +702,7 @@ namespace ProjectLauncher
 				.ButtonStyle(FProjectLauncherStyle::Get(), "HoverHintOnly")
 				.ToolTipText(LOCTEXT("FileBrowseTip", "Browse for a file"))
 				.OnClicked_Lambda(OnBrowse)
+				.IsEnabled_Lambda(IsEnabled)
 				.ContentPadding(2)
 				[
 					SNew(SImage)
@@ -714,7 +735,6 @@ namespace ProjectLauncher
 		FCallbacks WidgetCallbacks
 		{
 			.IsVisible = StringCallbacks.IsVisible,
-			.IsEnabled = StringCallbacks.IsEnabled,
 			.Validation = StringCallbacks.Validation,
 		};
 		if (StringCallbacks.GetDefaultValue != nullptr)
@@ -799,6 +819,11 @@ namespace ProjectLauncher
 			return EVisibility::Collapsed;
 		};
 
+		auto IsEnabled = [StringCallbacks]()
+		{
+			return !StringCallbacks.IsEnabled.IsSet() || StringCallbacks.IsEnabled();
+		};
+
 		return AddWidget( InName,
 			MoveTemp(WidgetCallbacks),
 			SNew(SHorizontalBox)
@@ -814,6 +839,7 @@ namespace ProjectLauncher
 				.Font(FCoreStyle::Get().GetFontStyle("SmallFont"))
 				.Text_Lambda( GetString )
 				.OnTextCommitted_Lambda( SetString )
+				.IsEnabled_Lambda(IsEnabled)
 			]
 
 			// command line parameters button
@@ -828,6 +854,7 @@ namespace ProjectLauncher
 				.OnGetMenuContent_Lambda(OnGetCmdlineParameterMenuContent)
 				.Visibility_Lambda(GetExtensionButtonVisibility)
 				.MenuPlacement(MenuPlacement_BelowRightAnchor)
+				.IsEnabled_Lambda(IsEnabled)
 				.HasDownArrow(false)
 				.ButtonContent()
 				[
diff --git a/Engine/Plugins/Developer/ProjectLauncher/Source/ProjectLauncher/Public/ProfileTree/LaunchProfileTreeData.h b/Engine/Plugins/Developer/ProjectLauncher/Source/ProjectLauncher/Public/ProfileTree/LaunchProfileTreeData.h
--- a/Engine/Plugins/Developer/ProjectLauncher/Source/ProjectLauncher/Public/ProfileTree/LaunchProfileTreeData.h
+++ b/Engine/Plugins/Developer/ProjectLauncher/Source/ProjectLauncher/Public/ProfileTree/LaunchProfileTreeData.h
@@ -239,7 +239,6 @@ namespace ProjectLauncher
 		FCallbacks WidgetCallbacks
 		{
 			.IsVisible = TypeCallbacks.IsVisible,
-			.IsEnabled = TypeCallbacks.IsEnabled,
 			.Validation = TypeCallbacks.Validation,
 		};
 		if (TypeCallbacks.GetDefaultValue != nullptr)
@@ -279,6 +278,11 @@ namespace ProjectLauncher
 			return TypeCallbacks.GetValue();
 		};
 
+		auto IsEnabled = [TypeCallbacks]()
+		{
+			return TypeCallbacks.IsEnabled && TypeCallbacks.IsEnabled();
+		};
+
 		if (Range.IsValid())
 		{
 			return AddWidget( InName,
@@ -290,6 +294,7 @@ namespace ProjectLauncher
 				.MinValue(Range.Min)
 				.MaxValue(Range.Max)
 				.EnableSlider(bShowSpinner)
+				.IsEnabled_Lambda(IsEnabled)
 			);
 		}
 		else
@@ -300,6 +305,7 @@ namespace ProjectLauncher
 				.Value_Lambda(GetValue)
 				.OnValueCommitted_Lambda(SetValue)
 				.EnableSlider(false)
+				.IsEnabled_Lambda(IsEnabled)
 			);
 		}
 	}
