client/server command line parameters can use a proper command line editor field

#rnx

#Synced-CL 51007224

diff --git a/Engine/Plugins/Developer/ProjectLauncher/Source/ProjectLauncher/Private/ProfileTree/GenericProfileTreeBuilder.cpp b/Engine/Plugins/Developer/ProjectLauncher/Source/ProjectLauncher/Private/ProfileTree/GenericProfileTreeBuilder.cpp
--- a/Engine/Plugins/Developer/ProjectLauncher/Source/ProjectLauncher/Private/ProfileTree/GenericProfileTreeBuilder.cpp
+++ b/Engine/Plugins/Developer/ProjectLauncher/Source/ProjectLauncher/Private/ProfileTree/GenericProfileTreeBuilder.cpp
@@ -801,7 +801,7 @@ namespace ProjectLauncher
 		{
 			EBuildTargetType TargetType = PerTargetTypeCmdLine[Index].TargetType;
 
-			HeadingNode.AddString(  PerTargetTypeCmdLine[Index].Label, // cannot use .AddCommandLineString at the moment because extensions can cache state so all + menus only affect common command line for now
+			HeadingNode.AddCommandLineString(  PerTargetTypeCmdLine[Index].Label,
 				{
 					.GetValue = [this, TargetType]()				{ return GetCommandLine(TargetType); },
 					.SetValue = [this, TargetType](FString Value)	{ SetCommandLine(Value, TargetType); },
