fix for int and float project properties being incorrectly disabled

#rnx

#Synced-CL 51007224

diff --git a/Engine/Plugins/Developer/ProjectLauncher/Source/ProjectLauncher/Public/ProfileTree/LaunchProfileTreeData.h b/Engine/Plugins/Developer/ProjectLauncher/Source/ProjectLauncher/Public/ProfileTree/LaunchProfileTreeData.h
--- a/Engine/Plugins/Developer/ProjectLauncher/Source/ProjectLauncher/Public/ProfileTree/LaunchProfileTreeData.h
+++ b/Engine/Plugins/Developer/ProjectLauncher/Source/ProjectLauncher/Public/ProfileTree/LaunchProfileTreeData.h
@@ -280,7 +280,7 @@ namespace ProjectLauncher
 
 		auto IsEnabled = [TypeCallbacks]()
 		{
-			return TypeCallbacks.IsEnabled && TypeCallbacks.IsEnabled();
+			return !TypeCallbacks.IsEnabled.IsSet() || TypeCallbacks.IsEnabled();
 		};
 
 		if (Range.IsValid())
