Move the cooked platform validation to the model, and it's only applied when the selected profile is actually using the fallback project.

#rnx

#Synced-CL 51007224

diff --git a/Engine/Plugins/Developer/ProjectLauncher/Source/ProjectLauncher/Private/Model/ProjectLauncherModel.cpp b/Engine/Plugins/Developer/ProjectLauncher/Source/ProjectLauncher/Private/Model/ProjectLauncherModel.cpp
--- a/Engine/Plugins/Developer/ProjectLauncher/Source/ProjectLauncher/Private/Model/ProjectLauncherModel.cpp
+++ b/Engine/Plugins/Developer/ProjectLauncher/Source/ProjectLauncher/Private/Model/ProjectLauncherModel.cpp
@@ -55,6 +55,7 @@ namespace ProjectLauncher
 
 		ProfileManager->OnProfileAdded().AddRaw(this, &FModel::HandleProfileManagerProfileAdded);
 		ProfileManager->OnProfileRemoved().AddRaw(this, &FModel::HandleProfileManagerProfileRemoved);
+		ProfileManager->OnProjectChanged().AddRaw(this, &FModel::HandleProfileManagerFallbackProjectChanged);
 
 		DeviceProxyManager->OnProxyAdded().AddRaw(this, &FModel::HandleDeviceProxyAdded);
 		DeviceProxyManager->OnProxyRemoved().AddRaw(this, &FModel::HandleDeviceProxyRemoved);
@@ -109,6 +110,8 @@ namespace ProjectLauncher
 
 		ProfileManager->OnProfileAdded().RemoveAll(this);
 		ProfileManager->OnProfileRemoved().RemoveAll(this);
+		ProfileManager->OnProjectChanged().RemoveAll(this);
+
 
 		DeviceProxyManager->OnProxyAdded().RemoveAll(this);
 		DeviceProxyManager->OnProxyRemoved().RemoveAll(this);
@@ -483,6 +486,16 @@ PRAGMA_ENABLE_DEPRECATION_WARNINGS
 		}
 	}
 
+	void FModel::HandleProfileManagerFallbackProjectChanged()
+	{
+		if (SelectedProfile.IsValid() && !SelectedProfile->HasProjectSpecified())
+		{
+			for (const ILauncherProfileBuildCookRunRef& BuildCookRun : SelectedProfile->GetBuildCookRunCommands())
+			{
+				UpdateCookedPlatformsFromBuildTarget(SelectedProfile.ToSharedRef(), BuildCookRun);
+			}
+		}
+	}
 
 
 
diff --git a/Engine/Plugins/Developer/ProjectLauncher/Source/ProjectLauncher/Private/Widgets/Profiles/SCustomLaunchProfilesPanel.cpp b/Engine/Plugins/Developer/ProjectLauncher/Source/ProjectLauncher/Private/Widgets/Profiles/SCustomLaunchProfilesPanel.cpp
--- a/Engine/Plugins/Developer/ProjectLauncher/Source/ProjectLauncher/Private/Widgets/Profiles/SCustomLaunchProfilesPanel.cpp
+++ b/Engine/Plugins/Developer/ProjectLauncher/Source/ProjectLauncher/Private/Widgets/Profiles/SCustomLaunchProfilesPanel.cpp
@@ -900,14 +900,6 @@ void SCustomLaunchProfilesPanel::SetDefaultProjectPath(FString ProjectPath)
 {
 	Model->GetProfileManager()->SetProjectPath(ProjectPath);
 	Model->GetProfileManager()->SetBuildTarget(FString());
-
-	if (Model->GetSelectedProfile().IsValid())
-	{
-		for (const ILauncherProfileBuildCookRunRef& BuildCookRun : Model->GetSelectedProfile()->GetBuildCookRunCommands())
-		{
-			Model->UpdateCookedPlatformsFromBuildTarget(Model->GetSelectedProfile().ToSharedRef(), BuildCookRun);
-		}
-	}
 }
 
 FString SCustomLaunchProfilesPanel::GetDefaultProjectPath() const
diff --git a/Engine/Plugins/Developer/ProjectLauncher/Source/ProjectLauncher/Public/Model/ProjectLauncherModel.h b/Engine/Plugins/Developer/ProjectLauncher/Source/ProjectLauncher/Public/Model/ProjectLauncherModel.h
--- a/Engine/Plugins/Developer/ProjectLauncher/Source/ProjectLauncher/Public/Model/ProjectLauncherModel.h
+++ b/Engine/Plugins/Developer/ProjectLauncher/Source/ProjectLauncher/Public/Model/ProjectLauncherModel.h
@@ -321,6 +321,7 @@ namespace ProjectLauncher
 		/* Handlers for profile manager and device proxy events */
 		UE_API void HandleProfileManagerProfileAdded(const ILauncherProfileRef& Profile);
 		UE_API void HandleProfileManagerProfileRemoved(const ILauncherProfileRef& Profile);
+		UE_API void HandleProfileManagerFallbackProjectChanged();
 		UE_API void HandleDeviceProxyAdded(const TSharedRef<ITargetDeviceProxy>& DeviceProxy);
 		UE_API void HandleDeviceProxyRemoved(const TSharedRef<ITargetDeviceProxy>& DeviceProxy);
 
