reduce the number of ini file branches that are loaded when reading other projects' config files

#rnx

#Synced-CL 51007224

diff --git a/Engine/Plugins/Developer/ProjectLauncher/Source/ProjectLauncher/Private/Model/ProjectLauncherModel.cpp b/Engine/Plugins/Developer/ProjectLauncher/Source/ProjectLauncher/Private/Model/ProjectLauncherModel.cpp
--- a/Engine/Plugins/Developer/ProjectLauncher/Source/ProjectLauncher/Private/Model/ProjectLauncherModel.cpp
+++ b/Engine/Plugins/Developer/ProjectLauncher/Source/ProjectLauncher/Private/Model/ProjectLauncherModel.cpp
@@ -979,7 +979,12 @@ PRAGMA_ENABLE_DEPRECATION_WARNINGS
 			Result.Config = new FConfigCacheIni(EConfigCacheType::Temporary);
 			FConfigContext Context = FConfigContext::ReadIntoConfigSystem(Result.Config, FString());
 			Context.ProjectConfigDir = FPaths::Combine(FPaths::GetPath(ProjectPath), TEXT("Config/"));
-			Result.Config->InitializeKnownConfigFiles(Context);
+
+			const TCHAR* IniNames[] = { *GEngineIni, *GGameIni }; // only need a subset of the other project's config files
+			for (const TCHAR* IniName : IniNames)
+			{
+				Context.Load(IniName);
+			}
 
 			// read the ini settings
 			Result.bIsCurrentEditorProject = false;
