remove duplicated IoStore property field - it is handled by the deprecated properties extension

#rnx

#Synced-CL 51007224

diff --git a/Engine/Plugins/Developer/ProjectLauncher/Source/ProjectLauncher/Private/ProfileTree/CustomProfileTreeBuilder.cpp b/Engine/Plugins/Developer/ProjectLauncher/Source/ProjectLauncher/Private/ProfileTree/CustomProfileTreeBuilder.cpp
--- a/Engine/Plugins/Developer/ProjectLauncher/Source/ProjectLauncher/Private/ProfileTree/CustomProfileTreeBuilder.cpp
+++ b/Engine/Plugins/Developer/ProjectLauncher/Source/ProjectLauncher/Private/ProfileTree/CustomProfileTreeBuilder.cpp
@@ -57,7 +57,6 @@ namespace ProjectLauncher
 
 			FLaunchProfileTreeNode& ContentSchemeHeader = AddHeading( GeneralBuildCookRunHeader, TEXT("ContentScheme"), LOCTEXT("ContentSchemeHeading", "Content Scheme") );
 			BuildCookRun.AddCompressPakFilesProperty(ContentSchemeHeader);
-			BuildCookRun.AddUseIoStoreProperty(ContentSchemeHeader);
 			BuildCookRun.AddGenerateChunksProperty(ContentSchemeHeader);
 			BuildCookRun.AddZenSnapshotProperty(ContentSchemeHeader);
 			BuildCookRun.AddImportZenSnapshotProperty(ContentSchemeHeader);
diff --git a/Engine/Plugins/Developer/ProjectLauncher/Source/ProjectLauncher/Private/ProfileTree/GenericProfileTreeBuilder.cpp b/Engine/Plugins/Developer/ProjectLauncher/Source/ProjectLauncher/Private/ProfileTree/GenericProfileTreeBuilder.cpp
--- a/Engine/Plugins/Developer/ProjectLauncher/Source/ProjectLauncher/Private/ProfileTree/GenericProfileTreeBuilder.cpp
+++ b/Engine/Plugins/Developer/ProjectLauncher/Source/ProjectLauncher/Private/ProfileTree/GenericProfileTreeBuilder.cpp
@@ -308,23 +308,6 @@ namespace ProjectLauncher
 		);
 	}
 
-	void FGenericProfileTreeBuilder::FBuildCookRun::AddUseIoStoreProperty( FLaunchProfileTreeNode& HeadingNode )
-	{
-		auto IoStorePropertyVisible = [this]()
-		{
-			return ForPak()  && !Model->GetProjectSettings(Owner.Profile).bUseIoStore;
-		};
-
-		HeadingNode.AddBoolean( LOCTEXT("UseIoStoreLabel", "Use Io Store"),
-			{
-				.GetValue = [this]()			{ return BuildCookRun->IsUsingIoStore(); },
-				.SetValue = [this](bool bValue)	{ BuildCookRun->SetUseIoStore(bValue); },
-				.GetDefaultValue = [this]()		{ return Owner.DefaultBuildCookRun->IsUsingIoStore(); },
-				.IsVisible = IoStorePropertyVisible,
-			}
-		);
-	}
-
 	void FGenericProfileTreeBuilder::FBuildCookRun::AddGenerateChunksProperty( FLaunchProfileTreeNode& HeadingNode )
 	{
 		HeadingNode.AddBoolean( LOCTEXT("GenerateChunksLabel", "Generate Chunks"),
diff --git a/Engine/Plugins/Developer/ProjectLauncher/Source/ProjectLauncher/Public/ProfileTree/GenericProfileTreeBuilder.h b/Engine/Plugins/Developer/ProjectLauncher/Source/ProjectLauncher/Public/ProfileTree/GenericProfileTreeBuilder.h
--- a/Engine/Plugins/Developer/ProjectLauncher/Source/ProjectLauncher/Public/ProfileTree/GenericProfileTreeBuilder.h
+++ b/Engine/Plugins/Developer/ProjectLauncher/Source/ProjectLauncher/Public/ProfileTree/GenericProfileTreeBuilder.h
@@ -66,7 +66,6 @@ namespace ProjectLauncher
 			UE_API void AddConfigurationProperty( FLaunchProfileTreeNode& HeadingNode );
 			UE_API void AddContentSchemeProperty( FLaunchProfileTreeNode& HeadingNode );
 			UE_API void AddCompressPakFilesProperty( FLaunchProfileTreeNode& HeadingNode );
-			UE_API void AddUseIoStoreProperty( FLaunchProfileTreeNode& HeadingNode );
 			UE_API void AddGenerateChunksProperty( FLaunchProfileTreeNode& HeadingNode );
 			UE_API void AddZenSnapshotProperty(FLaunchProfileTreeNode& HeadingNode);
 			UE_API void AddImportZenSnapshotProperty( FLaunchProfileTreeNode& HeadingNode );
