draw each separate UAT command in a different colour

#rnx

#Synced-CL 51094919

diff --git a/Engine/Plugins/Developer/ProjectLauncher/Source/ProjectLauncher/Private/Widgets/Output/CustomLaunchOutputLogMarshaller.cpp b/Engine/Plugins/Developer/ProjectLauncher/Source/ProjectLauncher/Private/Widgets/Output/CustomLaunchOutputLogMarshaller.cpp
--- a/Engine/Plugins/Developer/ProjectLauncher/Source/ProjectLauncher/Private/Widgets/Output/CustomLaunchOutputLogMarshaller.cpp
+++ b/Engine/Plugins/Developer/ProjectLauncher/Source/ProjectLauncher/Private/Widgets/Output/CustomLaunchOutputLogMarshaller.cpp
@@ -2,6 +2,7 @@
 
 #include "Widgets/Output/CustomLaunchOutputLogMarshaller.h"
 #include "Model/ProjectLauncherModel.h"
+#include "Math/ColorList.h"
 #include "Styling/AppStyle.h"
 #include "Framework/Text/SlateTextRun.h"
 #include "Framework/Text/ISlateLineHighlighter.h"
@@ -15,8 +16,22 @@ namespace ProjectLauncher
 	{
 		MessageStyle = FAppStyle::Get().GetWidgetStyle<FTextBlockStyle>("MonospacedText");
 
-		DisplayStyle = MessageStyle;
-		DisplayStyle.ColorAndOpacity = FSlateColor(FColor::Green);
+		static const FSlateColor DisplayColors[] =
+		{
+			FSlateColor(FColorList::White),
+			FSlateColor(FColorList::Green),
+			FSlateColor(FColorList::Magenta),
+			FSlateColor(FColorList::Cyan),
+			FSlateColor(FColorList::LimeGreen),
+			FSlateColor(FColorList::NeonPink),
+			FSlateColor(FColorList::NeonBlue),
+		};
+
+		for (const FSlateColor& DisplayColor : DisplayColors)
+		{
+			FTextBlockStyle& DisplayStyle = DisplayStyles.Add_GetRef(MessageStyle);
+			DisplayStyle.ColorAndOpacity = DisplayColor;
+		}
 
 		WarningStyle = MessageStyle;
 		WarningStyle.ColorAndOpacity = FSlateColor(FColor::Yellow);
@@ -28,6 +43,7 @@ namespace ProjectLauncher
 
 	void FLaunchLogTextLayoutMarshaller::SetText(const FString& SourceString, FTextLayout& TargetTextLayout)
 	{
+		DisplayStyleIndex = 0;
 		TextLayout = &TargetTextLayout;
 		FlushPendingLogMessages();
 	}
@@ -55,6 +71,7 @@ namespace ProjectLauncher
 	void FLaunchLogTextLayoutMarshaller::RefreshAllLogMessages()
 	{
 		MakeDirty();
+		DisplayStyleIndex = 0;
 
 		for (TSharedPtr<FLaunchLogMessage> Message : Model->LaunchLogMessages)
 		{
@@ -114,25 +131,68 @@ namespace ProjectLauncher
 	}
 
 
-	TArray<TSharedRef<IRun>> FLaunchLogTextLayoutMarshaller::GetRun(TSharedPtr<FLaunchLogMessage> Message) const
+	TArray<TSharedRef<IRun>> FLaunchLogTextLayoutMarshaller::GetRun(TSharedPtr<FLaunchLogMessage> Message)
 	{
-		const FTextBlockStyle* TextStyle = &MessageStyle;
-		if (Message->Verbosity == ELogVerbosity::Warning)
+		TArray<TSharedRef<IRun>> Runs;
+
+		if (Message->Verbosity == ELogVerbosity::Display)
 		{
-			TextStyle = &WarningStyle;
+			// special case for "Display" log verbosity (aka the RunUAT command line parameter)
+			// draw each separate UAT command in a different colour on the same line
+			const TCHAR* CommandLinePtr = Message->Message->GetCharArray().GetData();
+
+			FTextRange Range(0,0); 
+			while(ParseUATCommand(CommandLinePtr, Range))
+			{
+				Runs.Add(FSlateTextRun::Create(FRunInfo(), Message->Message, DisplayStyles[DisplayStyleIndex], Range));
+				DisplayStyleIndex = (DisplayStyleIndex+1) % DisplayStyles.Num();
+			}
 		}
-		else if (Message->Verbosity == ELogVerbosity::Error || Message->Verbosity == ELogVerbosity::Fatal)
+		else
 		{
-			TextStyle = &ErrorStyle;
+			const FTextBlockStyle* TextStyle = &MessageStyle;
+
+			if (Message->Verbosity == ELogVerbosity::Warning)
+			{
+				TextStyle = &WarningStyle;
+			}
+			else if (Message->Verbosity == ELogVerbosity::Error || Message->Verbosity == ELogVerbosity::Fatal)
+			{
+				TextStyle = &ErrorStyle;
+			}
+
+			Runs.Add(FSlateTextRun::Create(FRunInfo(), Message->Message, *TextStyle));
 		}
-		else if (Message->Verbosity == ELogVerbosity::Display)
+
+		return MoveTemp(Runs);
+	}
+	
+	bool FLaunchLogTextLayoutMarshaller::ParseUATCommand( const TCHAR*& CommandLinePtr, FTextRange& TextRange ) const
+	{
+		const TCHAR* StartCommandLinePtr = CommandLinePtr;
+
+		const TCHAR* KnownPrefix = TEXT("Parsing command line: ");
+		const int32 KnownPrefixLength = FCString::Strlen(KnownPrefix);
+		if (FCString::Strnicmp(CommandLinePtr, KnownPrefix, KnownPrefixLength) == 0)
 		{
-			TextStyle = &DisplayStyle;
+			CommandLinePtr += KnownPrefixLength;
 		}
 
-		TArray<TSharedRef<IRun>> Runs;
-		Runs.Add(FSlateTextRun::Create(FRunInfo(), Message->Message, *TextStyle));
+		const TCHAR* PrevCommandLinePtr = nullptr;
+		FString Token;
+		while (CommandLinePtr != nullptr && *CommandLinePtr != TEXT('\0') && FParse::Token(CommandLinePtr, Token, false))
+		{
+			if (PrevCommandLinePtr != nullptr && !Token.StartsWith(TEXT("-")))
+			{
+				CommandLinePtr = PrevCommandLinePtr;
+				break;
+			}
 
-		return MoveTemp(Runs);
+			PrevCommandLinePtr = CommandLinePtr;
+		}
+
+		TextRange.BeginIndex = TextRange.EndIndex;
+		TextRange.EndIndex = TextRange.BeginIndex + (CommandLinePtr - StartCommandLinePtr);
+		return TextRange.EndIndex > TextRange.BeginIndex;
 	}
 }
diff --git a/Engine/Plugins/Developer/ProjectLauncher/Source/ProjectLauncher/Private/Widgets/Output/CustomLaunchOutputLogMarshaller.h b/Engine/Plugins/Developer/ProjectLauncher/Source/ProjectLauncher/Private/Widgets/Output/CustomLaunchOutputLogMarshaller.h
--- a/Engine/Plugins/Developer/ProjectLauncher/Source/ProjectLauncher/Private/Widgets/Output/CustomLaunchOutputLogMarshaller.h
+++ b/Engine/Plugins/Developer/ProjectLauncher/Source/ProjectLauncher/Private/Widgets/Output/CustomLaunchOutputLogMarshaller.h
@@ -45,12 +45,13 @@ namespace ProjectLauncher
 		void RefreshAllLogMessages();
 		int32 GetNumFilteredMessages() const { return NumFilteredMessages; }
 	protected:
-		TArray<TSharedRef<class IRun>> GetRun(TSharedPtr<FLaunchLogMessage> Message) const;
+		TArray<TSharedRef<class IRun>> GetRun(TSharedPtr<FLaunchLogMessage> Message);
+		bool ParseUATCommand( const TCHAR*& CommandLinePtr, FTextRange& TextRange ) const;
 
 		FTextBlockStyle MessageStyle;
-		FTextBlockStyle DisplayStyle;
 		FTextBlockStyle WarningStyle;
 		FTextBlockStyle ErrorStyle;
+		TArray<FTextBlockStyle> DisplayStyles;
 
 		const TSharedRef<ProjectLauncher::FModel> Model;
 		TSpscQueue<TSharedPtr<FLaunchLogMessage>> PendingMessages;
@@ -58,6 +59,7 @@ namespace ProjectLauncher
 		FString LogFilterString;
 		FTextLayout* TextLayout;
 		int32 NumFilteredMessages = 0;
+		int32 DisplayStyleIndex = 0;
 
 	};
 }
