Project Launcher: don't disable BuildUAT on a freshly created profile.

#rnx
#Synced-CL 51155769

diff --git a/Engine/Plugins/Developer/ProjectLauncher/Source/ProjectLauncher/Private/Model/ProjectLauncherModel.cpp b/Engine/Plugins/Developer/ProjectLauncher/Source/ProjectLauncher/Private/Model/ProjectLauncherModel.cpp
--- a/Engine/Plugins/Developer/ProjectLauncher/Source/ProjectLauncher/Private/Model/ProjectLauncherModel.cpp
+++ b/Engine/Plugins/Developer/ProjectLauncher/Source/ProjectLauncher/Private/Model/ProjectLauncherModel.cpp
@@ -537,7 +537,6 @@ PRAGMA_ENABLE_DEPRECATION_WARNINGS
 		ILauncherProfileRef Profile = ProfileManager->CreateUnsavedProfile(Name);
 
 		//set defaults
-		Profile->SetBuildUAT(false);
 		if (!ProfileManager->GetProjectPath().IsEmpty())
 		{
 			Profile->SetProjectPath(ProfileManager->GetProjectPath());
