UFE startup & project switching optimization.
- project settings are loaded in the background & dependent fields will disable themselves until the data is available

#rnx

#Synced-CL 51169632

diff --git a/Engine/Plugins/Developer/ProjectLauncher/Source/ProjectLauncher/Private/Model/ProjectLauncherModel.cpp b/Engine/Plugins/Developer/ProjectLauncher/Source/ProjectLauncher/Private/Model/ProjectLauncherModel.cpp
--- a/Engine/Plugins/Developer/ProjectLauncher/Source/ProjectLauncher/Private/Model/ProjectLauncherModel.cpp
+++ b/Engine/Plugins/Developer/ProjectLauncher/Source/ProjectLauncher/Private/Model/ProjectLauncherModel.cpp
@@ -6,6 +6,7 @@
 #include "Misc/ConfigContext.h"
 #include "Misc/App.h"
 #include "Misc/ScopedSlowTask.h"
+#include "Misc/ScopeLock.h"
 #include "PlatformInfo.h"
 #include "GameProjectHelper.h"
 #include "ITargetDeviceProxy.h"
@@ -17,6 +18,8 @@
 #include "Modules/ModuleManager.h"
 #include "Settings/ProjectPackagingSettings.h"
 #include "Experimental/ZenServerInterface.h"
+#include "HAL/PlatformProcess.h"
+#include "Async/Async.h"
 
 #if WITH_EDITOR
 #include "Engine/World.h"
@@ -43,6 +46,10 @@ namespace ProjectLauncher
 		: DeviceProxyManager(InDeviceProxyManager)
 		, Launcher(InLauncher)
 		, ProfileManager(InProfileManager)
+	{
+	}
+
+	void FModel::Init()
 	{
 		// use a custom ini file so that it can be shared between UnrealFrontend and the editor
 		ConfigFileName = FPaths::EngineSavedDir() / TEXT("Config") / TEXT("ProjectLauncher") / TEXT("UserSettings.ini");
@@ -60,12 +67,15 @@ namespace ProjectLauncher
 		DeviceProxyManager->OnProxyAdded().AddRaw(this, &FModel::HandleDeviceProxyAdded);
 		DeviceProxyManager->OnProxyRemoved().AddRaw(this, &FModel::HandleDeviceProxyRemoved);
 
+		ProjectSettingsReady.AddRaw(this, &FModel::HandleProjectSettingsReady);
+
 		// ensure there's a project when we're in the editor
 		// this means we don't need to display the global project selector in the editor
 		if (GIsEditor && FPaths::IsProjectFilePathSet())
 		{
 			ProfileManager->SetProjectPath(FPaths::GetProjectFilePath());
 		}
+		GetProjectSettings(ProfileManager->GetProjectPath(), false); // begin caching the current project settings
 
 		// prepare profiles
 		BasicLaunchProfile = CreateBasicLaunchProfile();
@@ -74,6 +84,8 @@ namespace ProjectLauncher
 		AllProfiles = ProfileManager->GetAllProfiles();
 		for (const ILauncherProfilePtr& Profile : AllProfiles)
 		{
+			Profile->OnProjectChanged().AddRaw(this, &FModel::HandleProfileProjectChanged, Profile);
+
 			Profile->SetBuildUAT(false); // this isn't serialized, so set a better default
 		}
 
@@ -102,9 +114,15 @@ namespace ProjectLauncher
 	{
 		for (TPair<FString, FProjectSettings>& Pair : CachedProjectSettings)
 		{
+			FPlatformProcess::ReturnSynchEventToPool(Pair.Value.Pending);
 			delete Pair.Value.Config;
 		}
 
+		for (const ILauncherProfilePtr& Profile : AllProfiles)
+		{
+			Profile->OnProjectChanged().RemoveAll(this);
+		}
+
 		ProfileManager->OnPostProcessLaunchCommandLine().RemoveAll(this);
 		ProfileManager->OnPostProcessUATCommandLine().RemoveAll(this);
 
@@ -469,6 +487,8 @@ PRAGMA_ENABLE_DEPRECATION_WARNINGS
 	void FModel::HandleProfileManagerProfileAdded(const ILauncherProfileRef& Profile)
 	{
 		AllProfiles.Add(Profile);
+
+		Profile->OnProjectChanged().AddRaw(this, &FModel::HandleProfileProjectChanged, Profile.ToSharedPtr());
 	}
 
 	extern void RemoveExtensionInstancesForProfile( const ILauncherProfileRef& Profile );
@@ -476,6 +496,8 @@ PRAGMA_ENABLE_DEPRECATION_WARNINGS
 
 	void FModel::HandleProfileManagerProfileRemoved(const ILauncherProfileRef& Profile)
 	{
+		Profile->OnProjectChanged().RemoveAll(this);
+
 		RemoveExtensionInstancesForProfile(Profile);
 
 		AllProfiles.Remove(Profile);
@@ -488,7 +510,7 @@ PRAGMA_ENABLE_DEPRECATION_WARNINGS
 
 	void FModel::HandleProfileManagerFallbackProjectChanged()
 	{
-		if (SelectedProfile.IsValid() && !SelectedProfile->HasProjectSpecified())
+		if (SelectedProfile.IsValid() && !SelectedProfile->HasProjectSpecified() && AreProjectSettingsReady(SelectedProfile.ToSharedRef()))
 		{
 			for (const ILauncherProfileBuildCookRunRef& BuildCookRun : SelectedProfile->GetBuildCookRunCommands())
 			{
@@ -497,6 +519,14 @@ PRAGMA_ENABLE_DEPRECATION_WARNINGS
 		}
 	}
 
+	void FModel::HandleProjectSettingsReady( const FString& ProjectPath )
+	{
+		if (ProjectPath == ProfileManager->GetProjectPath())
+		{
+			HandleProfileManagerFallbackProjectChanged();
+		}
+	}
+
 
 
 	void FModel::HandleDeviceProxyAdded(const TSharedRef<ITargetDeviceProxy>& DeviceProxy )
@@ -512,7 +542,28 @@ PRAGMA_ENABLE_DEPRECATION_WARNINGS
 	{
 	}
 
+	void FModel::HandleProfileProjectChanged( ILauncherProfilePtr InProfile )
+	{
+		if (InProfile.IsValid())
+		{
+			FString GlobalProjectPath;
+			if (FPaths::IsProjectFilePathSet())
+			{
+				GlobalProjectPath = FPaths::ConvertRelativePathToFull(FPaths::GetProjectFilePath());
+			}
 
+			FString ProjectPath = FPaths::ConvertRelativePathToFull(InProfile->GetProjectPath());
+			if (!GlobalProjectPath.IsEmpty() && GlobalProjectPath == ProjectPath )
+			{
+				// this is the current project in-editor project - not for UFE. The build target list will have been loaded already and all ini files are already in memory
+			}
+			else
+			{
+				// begin caching the new project settings
+				GetProjectSettings(InProfile->GetProjectPath(), false);
+			}
+		}
+	}
 
 
 	TSharedPtr<ITargetDeviceProxy> FModel::GetDeviceProxy(const ILauncherProfileBuildCookRunRef& BuildCookRun)
@@ -722,7 +773,7 @@ PRAGMA_ENABLE_DEPRECATION_WARNINGS
 
 	void FModel::EnsureValidBuildTarget(const ILauncherProfileRef& Profile, const ILauncherProfileBuildCookRunRef& BuildCookRun)
 	{
-		if (Profile->HasProjectSpecified())
+		if (Profile->HasProjectSpecified() && AreProjectSettingsReady(Profile))
 		{
 			const FProjectSettings& ProjectSettings = GetProjectSettings(Profile);
 			if (!ProjectSettings.FallbackDefaultBuildTarget.IsEmpty() && (BuildCookRun->GetBuildTargets().Num() == 0 || BuildCookRun->GetBuildTargets().Contains(FString())))
@@ -938,10 +989,25 @@ PRAGMA_ENABLE_DEPRECATION_WARNINGS
 #endif
 
 	}
+	
+	
+	bool FModel::AreProjectSettingsReady(const FString& ProjectPath)
+	{
+		FScopeLock Lock(&CachedProjectSettingsCS);
+		return GetProjectSettings(ProjectPath, false).bIsReady;
+	}
 
-	const FProjectSettings FModel::GetProjectSettings( const FString& InProjectPath )
+	bool FModel::AreProjectSettingsReady(const ILauncherProfileRef& Profile)
+	{
+		return AreProjectSettingsReady(*Profile->GetProjectPath());
+	}
+	
+	
+
+	const FProjectSettings FModel::GetProjectSettings( const FString& InProjectPath, bool bWait )
 	{
 		FProjectSettings Result;
+		Result.bIsReady = true;
 
 		FString ProjectPath = FPaths::ConvertRelativePathToFull(InProjectPath);
 		FString ProjectName = FPaths::GetBaseFilename(ProjectPath);
@@ -970,37 +1036,92 @@ PRAGMA_ENABLE_DEPRECATION_WARNINGS
 
 			// note this is not cached because the user could edit the properties from within the editor
 		}
-		else if (CachedProjectSettings.Contains(ProjectName))
+		else if (!ProjectName.IsEmpty())
 		{
-			Result = CachedProjectSettings[ProjectName];
-		}
-		else if (!ProjectPath.IsEmpty())
-		{
-			// load the other project's ini files into a temporary config cache
-			Result.Config = new FConfigCacheIni(EConfigCacheType::Temporary);
-			FConfigContext Context = FConfigContext::ReadIntoConfigSystem(Result.Config, FString());
-			Context.ProjectConfigDir = FPaths::Combine(FPaths::GetPath(ProjectPath), TEXT("Config/"));
+			FEvent* WaitEvent = nullptr;
 
-			const TCHAR* IniNames[] = { *GEngineIni, *GGameIni }; // only need a subset of the other project's config files
-			for (const TCHAR* IniName : IniNames)
+			// find or add cached project settings
 			{
-				Context.Load(IniName);
+				FScopeLock Lock(&CachedProjectSettingsCS);
+
+				// check if it exists already
+				if (!CachedProjectSettings.Contains(ProjectName))
+				{
+					// create a placeholder entry
+					FProjectSettings& Settings = CachedProjectSettings.Emplace( ProjectName );
+					Settings.Pending = FPlatformProcess::GetSynchEventFromPool(true);
+					WaitEvent = Settings.Pending;
+
+					// initialize it in the background
+					AsyncTask(ENamedThreads::AnyHiPriThreadNormalTask, [WeakThis = AsWeak(), ProjectName, ProjectPath = InProjectPath]()
+					{
+						TSharedPtr<FModel> StrongThis = WeakThis.Pin();
+						if (StrongThis.IsValid())
+						{
+							FProjectSettings NewSettings;
+
+							// load the other project's ini files into a temporary config cache
+							NewSettings.Config = new FConfigCacheIni(EConfigCacheType::Temporary);
+							FConfigContext Context = FConfigContext::ReadIntoConfigSystem(NewSettings.Config, FString());
+							Context.ProjectConfigDir = FPaths::Combine(FPaths::GetPath(ProjectPath), TEXT("Config/"));
+
+							const TCHAR* IniNames[] = { *GEngineIni, *GGameIni }; // only need a subset of the other project's config files
+							for (const TCHAR* IniName : IniNames)
+							{
+								Context.Load(IniName);
+							}
+
+							// read the ini settings
+							NewSettings.bIsCurrentEditorProject = false;
+							StrongThis->ReadProjectSettingsFromConfig(*NewSettings.Config, ProjectPath, NewSettings);
+
+							// cache the result
+							{
+								FScopeLock Lock(&StrongThis->CachedProjectSettingsCS);
+								NewSettings.Pending = StrongThis->CachedProjectSettings[ProjectName].Pending;
+								StrongThis->CachedProjectSettings[ProjectName] = MoveTemp(NewSettings);
+								StrongThis->CachedProjectSettings[ProjectName].bIsReady = true;
+							}
+
+							// signal completion
+							StrongThis->CachedProjectSettings[ProjectName].Pending->Trigger();
+							ExecuteOnGameThread(UE_SOURCE_LOCATION, [WeakThis, ProjectPath]()
+							{
+								TSharedPtr<FModel> StrongThis = WeakThis.Pin();
+								if (StrongThis.IsValid())
+								{
+									StrongThis->ProjectSettingsReady.Broadcast(ProjectPath);
+								}
+							});
+						}
+					});
+				}
+			
+				WaitEvent = CachedProjectSettings[ProjectName].Pending;				
 			}
 
-			// read the ini settings
-			Result.bIsCurrentEditorProject = false;
-			ReadProjectSettingsFromConfig(*Result.Config, InProjectPath, Result);
+			// optionally wait for the settings to be ready
+			if (bWait && WaitEvent != nullptr)
+			{
+				WaitEvent->Wait();
+			}
 
-			// cache the result
-			CachedProjectSettings.Add( ProjectName, Result);
+			// grab the cached result
+			{
+				FScopeLock Lock(&CachedProjectSettingsCS);
+				Result = CachedProjectSettings[ProjectName];
+			}
 		}
 
 		return MoveTemp(Result);
 	}
-
-	const FProjectSettings FModel::GetProjectSettings( const ILauncherProfileRef& Profile )
+	
+	
+	
+	
+	const FProjectSettings FModel::GetProjectSettings( const ILauncherProfileRef& Profile, bool bWait )
 	{
-		return GetProjectSettings(*Profile->GetProjectPath());
+		return GetProjectSettings(*Profile->GetProjectPath(), bWait);
 	}
 
 	FString FModel::GetProjectPath(const ILauncherProfileRef& Profile) const
diff --git a/Engine/Plugins/Developer/ProjectLauncher/Source/ProjectLauncher/Private/ProfileTree/GenericProfileTreeBuilder.cpp b/Engine/Plugins/Developer/ProjectLauncher/Source/ProjectLauncher/Private/ProfileTree/GenericProfileTreeBuilder.cpp
--- a/Engine/Plugins/Developer/ProjectLauncher/Source/ProjectLauncher/Private/ProfileTree/GenericProfileTreeBuilder.cpp
+++ b/Engine/Plugins/Developer/ProjectLauncher/Source/ProjectLauncher/Private/ProfileTree/GenericProfileTreeBuilder.cpp
@@ -43,10 +43,14 @@ namespace ProjectLauncher
 
 		Profile->OnUATCommandAdded().AddRaw( this, &FGenericProfileTreeBuilder::OnUATCommandAdded );
 		Profile->OnUATCommandRemoved().AddRaw( this, &FGenericProfileTreeBuilder::OnUATCommandRemoved );
+
+		Model->OnProjectSettingsReady().AddRaw(this, &FGenericProfileTreeBuilder::ValidateBuildTargets );
+
 	}
 
 	FGenericProfileTreeBuilder::~FGenericProfileTreeBuilder()
 	{
+		Model->OnProjectSettingsReady().RemoveAll(this);
 		Profile->OnUATCommandAdded().RemoveAll(this);
 		Profile->OnUATCommandRemoved().RemoveAll(this);
 	}
@@ -115,11 +119,13 @@ namespace ProjectLauncher
 		};
 
 		Owner.Profile->OnCustomUATCommandValidation().AddRaw( this, &FBuildCookRun::OnValidateProfile );
+		Model->OnProjectSettingsReady().AddRaw( this, &FBuildCookRun::OnProjectSettingsReady );
 	}
 	
 	FGenericProfileTreeBuilder::FBuildCookRun::~FBuildCookRun()
 	{
 		Owner.Profile->OnCustomUATCommandValidation().RemoveAll( this );
+		Model->OnProjectSettingsReady().RemoveAll( this );
 	}
 
 
@@ -209,6 +215,7 @@ namespace ProjectLauncher
 				{
 					.IsDefault = [this]()		{ return !BuildCookRun->HasBuildTargetSpecified() || (BuildCookRun->GetBuildTargets().Num() == 0); },
 					.SetToDefault = [this]()	{ SetBuildTargets(TArray<FString>()); },
+					.IsEnabled = [this]()		{ return Model->AreProjectSettingsReady(Owner.Profile); },
 					.Validation = FValidation( {ELauncherProfileValidationErrors::BuildTargetCookVariantMismatch, ELauncherProfileValidationErrors::BuildTargetIsRequired, ELauncherProfileValidationErrors::CookOnTheFlyDoesntSupportServer} ),
 				},
 				SAssignNew(BuildTargetCombo, SCustomLaunchBuildTargetCombo, Model)
@@ -224,7 +231,7 @@ namespace ProjectLauncher
 				{
 					.IsDefault = [this]()		{ return !BuildCookRun->HasBuildTargetSpecified() || (BuildCookRun->GetBuildTargets().Num() == 0); },
 					.SetToDefault = [this]()	{ SetBuildTargets(TArray<FString>()); },
-					.IsEnabled = [this]()		{ return Owner.Profile->HasProjectSpecified(); },
+					.IsEnabled = [this]()		{ return Owner.Profile->HasProjectSpecified() && Model->AreProjectSettingsReady(Owner.Profile); },
 					.Validation = FValidation( {ELauncherProfileValidationErrors::BuildTargetCookVariantMismatch, ELauncherProfileValidationErrors::BuildTargetIsRequired, ELauncherProfileValidationErrors::CookOnTheFlyDoesntSupportServer}, {TEXT("Validation_GameAndClientTargetsDisallowed")}),
 				},
 				SAssignNew(BuildTargetCombo, SCustomLaunchBuildTargetCombo, Model)
@@ -242,7 +249,8 @@ namespace ProjectLauncher
 	{
 		HeadingNode.AddWidget( LOCTEXT("PlatformLabel", "Platform"), 
 			{
-				.Validation = FValidation({ELauncherProfileValidationErrors::NoPlatformSelected, ELauncherProfileValidationErrors::NoPlatformSDKInstalled})
+				.IsEnabled = [this]() { return Model->AreProjectSettingsReady(Owner.Profile); },
+				.Validation = FValidation({ELauncherProfileValidationErrors::NoPlatformSelected, ELauncherProfileValidationErrors::NoPlatformSDKInstalled}),
 			},
 			SAssignNew(PlatformCombo, SCustomLaunchPlatformCombo)
 			.SelectedPlatforms(this, &FGenericProfileTreeBuilder::FBuildCookRun::GetSelectedPlatforms)
@@ -825,6 +833,11 @@ namespace ProjectLauncher
 			return;
 		}
 
+		if (!Model->AreProjectSettingsReady(Owner.Profile))
+		{
+			return;
+		}
+
 		// verify that IoStore is available when a Zen Workflow is in use
 		if ((ContentScheme == EContentScheme::ZenStreaming || ContentScheme == EContentScheme::ZenPakStreaming))
 		{
@@ -845,6 +858,15 @@ namespace ProjectLauncher
 		}
 	}
 
+	void FGenericProfileTreeBuilder::FBuildCookRun::OnProjectSettingsReady(const FString& ProjectPath)
+	{
+		if (ProjectPath == Owner.Profile->GetProjectPath())
+		{
+			CacheBuildTargets();
+			OnValidateProfile(BuildCookRun);
+		}
+	}
+
 	void FGenericProfileTreeBuilder::OnUATCommandAdded(const ILauncherProfileUATCommandRef& UATCommand)
 	{
 		ILauncherProfileBuildCookRunPtr BuildCookRun = UATCommand->AsBuildCookRun();
@@ -937,15 +959,18 @@ namespace ProjectLauncher
 	{
 		CachedBuildTargetTypes.Reset();
 
-		TArray<FTargetInfo> TargetInfos = Model->GetBuildTargetInfos(Owner.Profile, BuildCookRun);
-		for (const FTargetInfo& TargetInfo : TargetInfos)
+		if (Model->AreProjectSettingsReady(Owner.Profile))
 		{
-			CachedBuildTargetTypes.Add(TargetInfo.Type);
-		}
+			TArray<FTargetInfo> TargetInfos = Model->GetBuildTargetInfos(Owner.Profile, BuildCookRun);
+			for (const FTargetInfo& TargetInfo : TargetInfos)
+			{
+				CachedBuildTargetTypes.Add(TargetInfo.Type);
+			}
 
-		if (bHasAdvancedPlatformTargets && PlatformCombo.IsValid())
-		{
-			PlatformCombo->RefreshPlatformsList();
+			if (bHasAdvancedPlatformTargets && PlatformCombo.IsValid())
+			{
+				PlatformCombo->RefreshPlatformsList();
+			}
 		}
 	}
 
@@ -1032,14 +1057,9 @@ namespace ProjectLauncher
 		Profile->SetProjectPath(ProjectPath);
 		Profile->SetProjectSpecified(!ProjectPath.IsEmpty());
 
-		for (const ILauncherProfileBuildCookRunRef& BuildCookRun : Profile->GetBuildCookRunCommands())
+		if (Model->AreProjectSettingsReady(Profile))
 		{
-			Model->EnsureValidBuildTarget(Profile, BuildCookRun);
-
-			if (ProfileType != EProfileType::Basic )
-			{
-				Model->UpdateCookedPlatformsFromBuildTarget(Profile, BuildCookRun);
-			}
+			ValidateBuildTargets(ProjectPath);
 		}
 
 		OnPropertyChanged();
@@ -1052,6 +1072,22 @@ namespace ProjectLauncher
 		return Profile->HasProjectSpecified();
 	}
 
+	void FGenericProfileTreeBuilder::ValidateBuildTargets(const FString& InProjectPath)
+	{
+		if (InProjectPath == Profile->GetProjectPath())
+		{
+			for (const ILauncherProfileBuildCookRunRef& BuildCookRun : Profile->GetBuildCookRunCommands())
+			{
+				Model->EnsureValidBuildTarget(Profile, BuildCookRun);
+
+				if (ProfileType != EProfileType::Basic )
+				{
+					Model->UpdateCookedPlatformsFromBuildTarget(Profile, BuildCookRun);
+				}
+			}
+			OnPropertyChanged();
+		}
+	}
 
 	TArray<FString> FGenericProfileTreeBuilder::FBuildCookRun::GetBuildTargets() const
 	{
diff --git a/Engine/Plugins/Developer/ProjectLauncher/Source/ProjectLauncher/Private/Widgets/Profiles/SCustomLaunchProfilesPanel.cpp b/Engine/Plugins/Developer/ProjectLauncher/Source/ProjectLauncher/Private/Widgets/Profiles/SCustomLaunchProfilesPanel.cpp
--- a/Engine/Plugins/Developer/ProjectLauncher/Source/ProjectLauncher/Private/Widgets/Profiles/SCustomLaunchProfilesPanel.cpp
+++ b/Engine/Plugins/Developer/ProjectLauncher/Source/ProjectLauncher/Private/Widgets/Profiles/SCustomLaunchProfilesPanel.cpp
@@ -187,6 +187,7 @@ TSharedRef<SWidget> SCustomLaunchProfilesPanel::CreateProfilesPanel()
 				.OnSelectionChanged(this, &SCustomLaunchProfilesPanel::SetDefaultBuildTarget)
 				.SelectedBuildTargets(this, &SCustomLaunchProfilesPanel::GetDefaultBuildTarget)
 				.SelectedProject(this, &SCustomLaunchProfilesPanel::GetDefaultProjectPath)
+				.IsEnabled_Lambda( [this]() { return !Model->GetSelectedProfile().IsValid() || Model->AreProjectSettingsReady(Model->GetSelectedProfile().ToSharedRef() ); } )
 			]
 
 			// Spacing
diff --git a/Engine/Plugins/Developer/ProjectLauncher/Source/ProjectLauncher/Private/Widgets/Shared/SCustomLaunchBuildTargetCombo.cpp b/Engine/Plugins/Developer/ProjectLauncher/Source/ProjectLauncher/Private/Widgets/Shared/SCustomLaunchBuildTargetCombo.cpp
--- a/Engine/Plugins/Developer/ProjectLauncher/Source/ProjectLauncher/Private/Widgets/Shared/SCustomLaunchBuildTargetCombo.cpp
+++ b/Engine/Plugins/Developer/ProjectLauncher/Source/ProjectLauncher/Private/Widgets/Shared/SCustomLaunchBuildTargetCombo.cpp
@@ -117,6 +117,11 @@ END_SLATE_FUNCTION_BUILD_OPTIMIZATION
 
 FText SCustomLaunchBuildTargetCombo::GetBuildTargetName() const
 {
+	if (SelectedProject.IsSet() && !Model->AreProjectSettingsReady(SelectedProject.Get()))
+	{
+		return FText::Format( LOCTEXT("TargetRefreshingLabel", "(loading {0} configuration...)"), FText::FromString(FPaths::GetBaseFilename(SelectedProject.Get())));
+	}
+
 	TArray<FString> BuildTargetNames = SelectedBuildTargets.Get();
 
 	if (BuildTargetNames.Num() > 1)
diff --git a/Engine/Plugins/Developer/ProjectLauncher/Source/ProjectLauncher/Public/Model/ProjectLauncherModel.h b/Engine/Plugins/Developer/ProjectLauncher/Source/ProjectLauncher/Public/Model/ProjectLauncherModel.h
--- a/Engine/Plugins/Developer/ProjectLauncher/Source/ProjectLauncher/Public/Model/ProjectLauncherModel.h
+++ b/Engine/Plugins/Developer/ProjectLauncher/Source/ProjectLauncher/Public/Model/ProjectLauncherModel.h
@@ -9,6 +9,7 @@
 #include "ILauncherWorker.h"
 #include "Templates/SharedPointer.h"
 #include "Misc/DataDrivenPlatformInfoRegistry.h"
+#include "HAL/Event.h"
 
 #define UE_API PROJECTLAUNCHER_API
 
@@ -113,6 +114,10 @@ namespace ProjectLauncher
 
 		// cached project config ini
 		FConfigCacheIni* Config = nullptr;
+
+		// whether this is still being initialized
+		FEvent* Pending = nullptr;
+		bool bIsReady = false;
 	};
 
 
@@ -122,6 +127,9 @@ namespace ProjectLauncher
 	/* Delegate fired when the selected profile changes */
 	DECLARE_MULTICAST_DELEGATE_TwoParams(FOnSelectedProfileChanged, const ILauncherProfilePtr& /*new*/, const ILauncherProfilePtr& /*old*/)
 
+	/* Delegate fired when project settings have been loaded */
+	DECLARE_MULTICAST_DELEGATE_OneParam(FOnProjectSettingsReady, const FString& /*ProjectPath*/)
+
 	/* Delegate fired when the selected platform changes */
 	DECLARE_MULTICAST_DELEGATE(FOnSelectedPlatformChanged)
 
@@ -140,9 +148,10 @@ namespace ProjectLauncher
 	class FModel : public TSharedFromThis<FModel>
 	{
 	public:
-		/* Constructor/destructor */
+		/* Constructor/destructor/one-time init */
 		UE_API FModel(const TSharedRef<ITargetDeviceProxyManager>& InDeviceProxyManager, const TSharedRef<ILauncher>& InLauncher, const TSharedRef<ILauncherProfileManager>& InProfileManager);
 		UE_API ~FModel();
+		UE_API void Init();
 
 		/* Get the target device proxy manager */
 		const TSharedRef<ITargetDeviceProxyManager>& GetDeviceProxyManager() const
@@ -281,10 +290,19 @@ namespace ProjectLauncher
 		static UE_API FString GetBuildTargetPlatformName(const FString& PlatformName, const FTargetInfo& BuildTargetInfo);
 
 		/* Retrieve project settings from a project path or profile */
-		UE_API const FProjectSettings GetProjectSettings(const FString& ProjectPath);
-		UE_API const FProjectSettings GetProjectSettings(const ILauncherProfileRef& Profile);
+		UE_API const FProjectSettings GetProjectSettings(const FString& ProjectPath, bool bWait = true);
+		UE_API const FProjectSettings GetProjectSettings(const ILauncherProfileRef& Profile, bool bWait = true);
+		UE_API bool AreProjectSettingsReady(const FString& ProjectPath);
+		UE_API bool AreProjectSettingsReady(const ILauncherProfileRef& Profile);
 		UE_API FString GetProjectPath(const ILauncherProfileRef& Profile) const;
 
+		/* Delegate fired when a project's settings have been loaded */
+		FOnProjectSettingsReady& OnProjectSettingsReady()
+		{
+			return ProjectSettingsReady;
+		}
+
+
 		/* Get config section and ini file used for storing project launcher state */
 		UE_API const TCHAR* GetConfigSection() const;
 		UE_API const FString& GetConfigIni() const;
@@ -325,6 +343,10 @@ namespace ProjectLauncher
 		UE_API void HandleDeviceProxyAdded(const TSharedRef<ITargetDeviceProxy>& DeviceProxy);
 		UE_API void HandleDeviceProxyRemoved(const TSharedRef<ITargetDeviceProxy>& DeviceProxy);
 
+		/* Handler for profile properties changing */
+		UE_API void HandleProfileProjectChanged( ILauncherProfilePtr InProfile );
+		UE_API void HandleProjectSettingsReady( const FString& ProjectPath );
+
 		/* Helper for reading project settings from config */
 		UE_API void ReadProjectSettingsFromConfig(FConfigCacheIni& InConfig, const FString& InProjectPath, FProjectSettings& OutResult);
 
@@ -350,11 +372,13 @@ namespace ProjectLauncher
 		friend class FLaunchLogTextLayoutMarshaller;
 		TArray<TSharedPtr<FLaunchLogMessage>> LaunchLogMessages;
 
+		FCriticalSection CachedProjectSettingsCS;
 		TMap<FString,FProjectSettings> CachedProjectSettings;
 		TMap<FString,TArray<FString>> CachedMapPaths;
 
 		FOnSelectedProfileChanged ProfileSelectedDelegate;
 		FOnSelectedPlatformChanged PlatformChangedDelegate;
+		FOnProjectSettingsReady ProjectSettingsReady;
 
 		ILauncherProfilePtr DefaultBasicLaunchProfile;
 		ILauncherProfilePtr DefaultCustomLaunchProfile;
diff --git a/Engine/Plugins/Developer/ProjectLauncher/Source/ProjectLauncher/Public/ProfileTree/GenericProfileTreeBuilder.h b/Engine/Plugins/Developer/ProjectLauncher/Source/ProjectLauncher/Public/ProfileTree/GenericProfileTreeBuilder.h
--- a/Engine/Plugins/Developer/ProjectLauncher/Source/ProjectLauncher/Public/ProfileTree/GenericProfileTreeBuilder.h
+++ b/Engine/Plugins/Developer/ProjectLauncher/Source/ProjectLauncher/Public/ProfileTree/GenericProfileTreeBuilder.h
@@ -53,7 +53,7 @@ namespace ProjectLauncher
 		UE_API FString GetProjectPath() const;
 		UE_API void SetProjectName(FString ProjectPath);
 		UE_API bool HasProject() const;
-
+		UE_API void ValidateBuildTargets(const FString& ProjectPath);
 
 		class FBuildCookRun : public TSharedFromThis<FBuildCookRun>
 		{
@@ -149,6 +149,7 @@ namespace ProjectLauncher
 			UE_API void RefreshContentScheme();
 			UE_API void ValidateArchitectures();
 			UE_API void OnValidateProfile(ILauncherProfileUATCommandRef UATCommand);
+			UE_API void OnProjectSettingsReady(const FString& ProjectPath);
 
 			UE_API void CacheArchitectures();
 			UE_API void CacheBuildTargets();
