add extra file missing from previous checkin

#rnx

#Synced-CL 51169632

diff --git a/Engine/Plugins/Developer/ProjectLauncher/Source/ProjectLauncher/Private/ProjectLauncherModule.cpp b/Engine/Plugins/Developer/ProjectLauncher/Source/ProjectLauncher/Private/ProjectLauncherModule.cpp
--- a/Engine/Plugins/Developer/ProjectLauncher/Source/ProjectLauncher/Private/ProjectLauncherModule.cpp
+++ b/Engine/Plugins/Developer/ProjectLauncher/Source/ProjectLauncher/Private/ProjectLauncherModule.cpp
@@ -103,6 +103,7 @@ private:
 		if (!Model.IsValid())
 		{
 			Model = MakeShared<ProjectLauncher::FModel>(TargetDeviceServicesModule.GetDeviceProxyManager(), LauncherServicesModule.CreateLauncher(), LauncherServicesModule.GetProfileManager() );
+			Model->Init();
 		}
 		
 		auto OnTabActivated = [this](TSharedRef<SDockTab>, ETabActivationCause Cause)
