Only add zen pdb dependencies if files exist

Fixes GitHub building of UE

#rnx
#Synced-CL 51374145

diff --git a/Engine/Plugins/Developer/ProjectLauncher/Source/CommonLaunchExtensions/CommonLaunchExtensions.Build.cs b/Engine/Plugins/Developer/ProjectLauncher/Source/CommonLaunchExtensions/CommonLaunchExtensions.Build.cs
--- a/Engine/Plugins/Developer/ProjectLauncher/Source/CommonLaunchExtensions/CommonLaunchExtensions.Build.cs
+++ b/Engine/Plugins/Developer/ProjectLauncher/Source/CommonLaunchExtensions/CommonLaunchExtensions.Build.cs
@@ -37,9 +37,19 @@ public class CommonLaunchExtensions : ModuleRules
 		{
 			RuntimeDependencies.Add("$(EngineDir)/Binaries/Win64/crashpad_handler.exe");
 			RuntimeDependencies.Add("$(EngineDir)/Binaries/Win64/zen.exe");
-			RuntimeDependencies.Add("$(EngineDir)/Binaries/Win64/zen.pdb");
+
+			if (System.IO.Directory.Exists(System.IO.Path.Combine(EngineDirectory, "Binaries", "Win64", "zen.pdb")))
+			{
+				RuntimeDependencies.Add("$(EngineDir)/Binaries/Win64/zen.pdb");
+			}
+
 			RuntimeDependencies.Add("$(EngineDir)/Binaries/Win64/zenserver.exe");
-			RuntimeDependencies.Add("$(EngineDir)/Binaries/Win64/zenserver.pdb");
+
+			if (System.IO.Directory.Exists(System.IO.Path.Combine(EngineDirectory, "Binaries", "Win64", "zenserver.pdb")))
+			{
+				RuntimeDependencies.Add("$(EngineDir)/Binaries/Win64/zenserver.pdb");
+			}
+
 			RuntimeDependencies.Add("$(EngineDir)/Binaries/DotNET/OidcToken/win-x64/OidcToken.exe");
 		}
 		else if (Target.Platform == UnrealTargetPlatform.Linux)
