ProjectLauncher add destination directory do BuildSync

#rnx
#Synced-CL 51818326

diff --git a/Engine/Plugins/Developer/ProjectLauncher/Source/CommonLaunchExtensions/Private/BuildSync/BuildSyncLaunchExtension.cpp b/Engine/Plugins/Developer/ProjectLauncher/Source/CommonLaunchExtensions/Private/BuildSync/BuildSyncLaunchExtension.cpp
--- a/Engine/Plugins/Developer/ProjectLauncher/Source/CommonLaunchExtensions/Private/BuildSync/BuildSyncLaunchExtension.cpp
+++ b/Engine/Plugins/Developer/ProjectLauncher/Source/CommonLaunchExtensions/Private/BuildSync/BuildSyncLaunchExtension.cpp
@@ -36,6 +36,13 @@ FBuildSyncLaunchExtensionInstance::FBuildSyncLaunchExtensionInstance( FArgs& InA
 	GetConfigString(EConfig::PerProfile, TEXT("NamedArtifacts")).ParseIntoArray(CachedNamedArtifacts, TEXT(";"));
 	CachedCookPlatforms = GetCookPlatforms();
 
+	DestinationDirectory = GetConfigString(EConfig::PerProfile, TEXT("DestinationDirectory"));
+	if (DestinationDirectory.IsEmpty())
+	{
+		FString CookedDir = FPaths::Combine(*FPaths::RootDir(), *InArgs.Profile->GetProjectBasePath(), TEXT("Saved"), TEXT("StagedBuilds"));
+		SetDestinationDirectory(CookedDir);
+	}
+
 	// create build info helper
 	BuildInfo = MakeShared<FBuildInfoHelper>();
 	BuildInfo->SetProjectName(GetProfile()->GetProjectName());
@@ -45,7 +52,9 @@ FBuildSyncLaunchExtensionInstance::FBuildSyncLaunchExtensionInstance( FArgs& InA
 	{
 		OnBuildsRefreshed();
 	};
-	BuildInfo->SetBuildsRefreshedHandler(OnRefreshed);	
+	BuildInfo->SetBuildsRefreshedHandler(OnRefreshed);
+
+	GetProfile()->OnUATCommandAdded().AddRaw(this, &FBuildSyncLaunchExtensionInstance::OnOtherUATCommandAdded);
 }
 
 FBuildSyncLaunchExtensionInstance::~FBuildSyncLaunchExtensionInstance()
@@ -53,6 +62,8 @@ FBuildSyncLaunchExtensionInstance::~FBuildSyncLaunchExtensionInstance()
 	BuildInfo->SetBuildsRefreshedHandler(nullptr);
 	BuildInfo.Reset();
 	BuildCombo.Reset();
+
+	GetProfile()->OnUATCommandAdded().RemoveAll(this);
 }
 
 
@@ -139,6 +150,15 @@ void FBuildSyncLaunchExtensionInstance::CustomizeTree( ProjectLauncher::FLaunchP
 			},
 			CreateNamedArtifactSelectionWidget()
 		)
+		.AddDirectoryString(LOCTEXT("DestinationDirectoryString", "Destination Directory"),
+			{
+				.GetValue = [this]() { return GetDestinationDirectory(); },
+				.SetValue = [this](FString InDestination) { SetDestinationDirectory(InDestination); },
+				.IsVisible = [this]() { return GetBuildType() != TEXT("oplog"); }
+			},
+			false,
+			LOCTEXT("DestinationDirectoryTooltip", "Will override the staging directory in BuildCookRun instances")
+		)
 	;
 }
 
@@ -544,6 +564,30 @@ TArray<FString> FBuildSyncLaunchExtensionInstance::GetNamedArtifacts() const
 
 
 
+void FBuildSyncLaunchExtensionInstance::SetDestinationDirectory(FString InPath)
+{
+	if (FPaths::IsRelative(InPath))
+	{
+		DestinationDirectory = FPaths::ConvertRelativePathToFull(InPath);
+	}
+	else
+	{
+		DestinationDirectory = InPath;
+	}
+	SetConfigString(EConfig::PerProfile, TEXT("DestinationDirectory"), DestinationDirectory);
+
+	TArray<ILauncherProfileBuildCookRunRef> BCRs = GetProfile()->GetBuildCookRunCommands();
+	for (auto BuildCookRun : BCRs)
+	{
+		BuildCookRun->SetPackageDirectory(DestinationDirectory);
+	}
+}
+
+FString FBuildSyncLaunchExtensionInstance::GetDestinationDirectory() const
+{
+	return DestinationDirectory;
+}
+
 void FBuildSyncLaunchExtensionInstance::AutoSelectBestBuild()
 {
 	const TArray<const PlatformInfo::FTargetPlatformInfo*> PlatformInfos = GetModel()->GetPlatformInfosForAllCommands(GetProfile());
@@ -776,6 +820,11 @@ void FBuildSyncLaunchExtensionInstance::CustomizeUATCommandLine( FString& InOutC
 	CommandLine += FString::Printf(TEXT(" -buildtype=%s"), *GetBuildType());
 	CommandLine += FString::Printf(TEXT(" -buildplatform=%s"), *FString::Join(CachedNamedArtifacts, TEXT("+")));
 
+	if (!DestinationDirectory.IsEmpty())
+	{
+		CommandLine += FString::Printf(TEXT(" -stagingdirectory=\"%s\""), *DestinationDirectory);
+	}
+
 	InOutCommandLine += CommandLine;
 }
 
@@ -795,6 +844,15 @@ void FBuildSyncLaunchExtensionInstance::OnUATCommandAdded( ILauncherProfileUATCo
 	InUATCommand->SetUATCommand(TEXT("BuildAcquire")); // RunUAT BuildAcquire [...] 
 }
 
+void FBuildSyncLaunchExtensionInstance::OnOtherUATCommandAdded(const ILauncherProfileUATCommandRef& InUATCommand)
+{
+	ILauncherProfileBuildCookRunPtr BuildCookRun = InUATCommand->AsBuildCookRun();
+	if (BuildCookRun.IsValid() && !DestinationDirectory.IsEmpty())
+	{
+		BuildCookRun->SetPackageDirectory(DestinationDirectory);
+	}
+}
+
 TArray<FString> FBuildSyncLaunchExtensionInstance::GetSelectedBuildBackends() const
 {
 	return SelectedBuild ? SelectedBuild->Backends : TArray<FString>();
diff --git a/Engine/Plugins/Developer/ProjectLauncher/Source/CommonLaunchExtensions/Private/DeprecatedProperties/DeprecatedPropertiesLaunchExtension.cpp b/Engine/Plugins/Developer/ProjectLauncher/Source/CommonLaunchExtensions/Private/DeprecatedProperties/DeprecatedPropertiesLaunchExtension.cpp
--- a/Engine/Plugins/Developer/ProjectLauncher/Source/CommonLaunchExtensions/Private/DeprecatedProperties/DeprecatedPropertiesLaunchExtension.cpp
+++ b/Engine/Plugins/Developer/ProjectLauncher/Source/CommonLaunchExtensions/Private/DeprecatedProperties/DeprecatedPropertiesLaunchExtension.cpp
@@ -113,13 +113,6 @@ void FDeprecatedPropertiesLaunchExtensionInstance::FBuildCookRunInstance::Custom
 			.SelectedItem_Lambda( [this]() { return GetBuildCookRun()->GetCookConfiguration(); } )
 			.Items(TArray<EBuildConfiguration>{ EBuildConfiguration::Debug, EBuildConfiguration::DebugGame, EBuildConfiguration::Development, EBuildConfiguration::Test, EBuildConfiguration::Shipping })
 		)
-		.AddDirectoryString( LOCTEXT("PackageDirLabel", "Staging Directory Override"),
-			{
-				.GetValue = [this]()				{ return GetBuildCookRun()->GetPackageDirectory(); },
-				.SetValue = [this](FString Value)	{ GetBuildCookRun()->SetPackageDirectory(Value); },
-				.GetDefaultValue = [this]()			{ return GetDefaultBuildCookRun()->GetPackageDirectory(); },
-			}
-		)
 		.AddBoolean( LOCTEXT("IoStoreLabel", "Use container files for optimized loading (I/O Store)"),
 			{
 				.GetValue = [this]()				{ return GetBuildCookRun()->IsUsingIoStore(); },
@@ -199,7 +192,6 @@ bool FDeprecatedPropertiesLaunchExtension::HasDeprecatedProperties(const ILaunch
 
 	return	InBuildCookRun->GetSkipCookingEditorContent()				!=	DefaultBuildCookRun->GetSkipCookingEditorContent() ||
 			InBuildCookRun->GetCookConfiguration()						!=	DefaultBuildCookRun->GetCookConfiguration() ||
-			InBuildCookRun->GetPackageDirectory()						!=	DefaultBuildCookRun->GetPackageDirectory() ||
 			InBuildCookRun->IsUsingIoStore()							!=	DefaultBuildCookRun->IsUsingIoStore() ||
 			InBuildCookRun->GetReferenceContainerGlobalFileName()		!=	DefaultBuildCookRun->GetReferenceContainerGlobalFileName() ||
 			InBuildCookRun->GetReferenceContainerCryptoKeysFileName()	!=	DefaultBuildCookRun->GetReferenceContainerCryptoKeysFileName()
diff --git a/Engine/Plugins/Developer/ProjectLauncher/Source/CommonLaunchExtensions/Public/BuildSync/BuildSyncLaunchExtension.h b/Engine/Plugins/Developer/ProjectLauncher/Source/CommonLaunchExtensions/Public/BuildSync/BuildSyncLaunchExtension.h
--- a/Engine/Plugins/Developer/ProjectLauncher/Source/CommonLaunchExtensions/Public/BuildSync/BuildSyncLaunchExtension.h
+++ b/Engine/Plugins/Developer/ProjectLauncher/Source/CommonLaunchExtensions/Public/BuildSync/BuildSyncLaunchExtension.h
@@ -27,6 +27,7 @@ public:
 	virtual void CustomizeUATCommandLine( FString& InOutCommandLine ) override;
 
 	virtual void OnUATCommandAdded( ILauncherProfileUATCommandRef UATCommand ) override;
+	void OnOtherUATCommandAdded(const ILauncherProfileUATCommandRef& UATCommand);
 
 
 	virtual void OnValidateProfile() override;
@@ -44,6 +45,8 @@ private:
 	FString GetBuildType() const;
 	void SetNamedArtifacts( TArray<FString> InNamedArtifacts );
 	TArray<FString> GetNamedArtifacts() const;
+	void SetDestinationDirectory(FString InPath);
+	FString GetDestinationDirectory() const;
 
 	void AutoSelectBestBuild();
 	void AutoSelectNamedArtifacts();
@@ -72,6 +75,7 @@ private:
 	bool bShowPreflights = false;
 	bool bShowUnsuitable = false;
 	FString FilterText;
+	FString DestinationDirectory;
 
 	TSet<const PlatformInfo::FTargetPlatformInfo*> GetCookPlatforms() const;
 	TSet<const PlatformInfo::FTargetPlatformInfo*> CachedCookPlatforms;
diff --git a/Engine/Plugins/Developer/ProjectLauncher/Source/ProjectLauncher/Private/ProfileTree/CustomProfileTreeBuilder.cpp b/Engine/Plugins/Developer/ProjectLauncher/Source/ProjectLauncher/Private/ProfileTree/CustomProfileTreeBuilder.cpp
--- a/Engine/Plugins/Developer/ProjectLauncher/Source/ProjectLauncher/Private/ProfileTree/CustomProfileTreeBuilder.cpp
+++ b/Engine/Plugins/Developer/ProjectLauncher/Source/ProjectLauncher/Private/ProfileTree/CustomProfileTreeBuilder.cpp
@@ -76,6 +76,7 @@ namespace ProjectLauncher
 			FLaunchProfileTreeNode& DirectoryHeader = AddHeading( GeneralBuildCookRunHeader,  TEXT("Directory"), LOCTEXT("DirectoryHeading", "Directory") );
 			BuildCookRun.AddArchiveBuildProperty(DirectoryHeader);
 			BuildCookRun.AddArchiveBuildDirectoryProperty(DirectoryHeader);
+			BuildCookRun.AddStagingDirectoryProperty(DirectoryHeader);
 
 			FLaunchProfileTreeNode& DeployAndRunHeader = AddHeading( GeneralBuildCookRunHeader,  TEXT("DeployAndRun"), LOCTEXT("DeployAndRunHeading", "Deploy And Run") );
 			BuildCookRun.AddDeployProperty(DeployAndRunHeader);
diff --git a/Engine/Plugins/Developer/ProjectLauncher/Source/ProjectLauncher/Private/ProfileTree/GenericProfileTreeBuilder.cpp b/Engine/Plugins/Developer/ProjectLauncher/Source/ProjectLauncher/Private/ProfileTree/GenericProfileTreeBuilder.cpp
--- a/Engine/Plugins/Developer/ProjectLauncher/Source/ProjectLauncher/Private/ProfileTree/GenericProfileTreeBuilder.cpp
+++ b/Engine/Plugins/Developer/ProjectLauncher/Source/ProjectLauncher/Private/ProfileTree/GenericProfileTreeBuilder.cpp
@@ -647,7 +647,7 @@ namespace ProjectLauncher
 
 	void FGenericProfileTreeBuilder::FBuildCookRun::AddStagingDirectoryProperty( FLaunchProfileTreeNode& HeadingNode )
 	{
-		HeadingNode.AddDirectoryString( LOCTEXT("CustomStagingPathLabel", "Custom Stage Directory"),
+		HeadingNode.AddDirectoryString( LOCTEXT("CustomStagingPathLabel", "Stage Directory Override"),
 			{
 				.GetValue = [this]()				{ return BuildCookRun->GetPackageDirectory(); },
 				.SetValue = [this](FString Value)	{ BuildCookRun->SetPackageDirectory(Value); },
diff --git a/Engine/Plugins/Developer/ProjectLauncher/Source/ProjectLauncher/Private/ProfileTree/LaunchProfileTreeData.cpp b/Engine/Plugins/Developer/ProjectLauncher/Source/ProjectLauncher/Private/ProfileTree/LaunchProfileTreeData.cpp
--- a/Engine/Plugins/Developer/ProjectLauncher/Source/ProjectLauncher/Private/ProfileTree/LaunchProfileTreeData.cpp
+++ b/Engine/Plugins/Developer/ProjectLauncher/Source/ProjectLauncher/Private/ProfileTree/LaunchProfileTreeData.cpp
@@ -510,7 +510,7 @@ namespace ProjectLauncher
 
 	}
 
-	FLaunchProfileTreeNode& FLaunchProfileTreeNode::AddDirectoryString( FText InName, FStringCallbacks&& StringCallbacks )
+	FLaunchProfileTreeNode& FLaunchProfileTreeNode::AddDirectoryString( FText InName, FStringCallbacks&& StringCallbacks, bool InMakeRelative, FText TooltipText )
 	{
 		check(StringCallbacks.GetValue);
 		check(StringCallbacks.SetValue);
@@ -534,6 +534,11 @@ namespace ProjectLauncher
 			};
 		};
 
+		auto GetTooltipText = [TooltipText]()
+		{
+			return TooltipText;
+		};
+
 		auto SetString = [this, StringCallbacks]( const FText& InText, ETextCommit::Type InTextCommit )
 		{
 			StringCallbacks.SetValue( InText.ToString() );
@@ -545,7 +550,7 @@ namespace ProjectLauncher
 			return FText::FromString(StringCallbacks.GetValue());
 		};
 
-		auto OnBrowse = [this,StringCallbacks]()
+		auto OnBrowse = [this,StringCallbacks, InMakeRelative]()
 		{
 			FString InitialDirectory = StringCallbacks.GetValue();
 			if (InitialDirectory.IsEmpty())
@@ -561,7 +566,7 @@ namespace ProjectLauncher
 			FString OutDirectory;
 			if (FDesktopPlatformModule::Get()->OpenDirectoryDialog(ParentWindowHandle, LOCTEXT("FolderDialogTitle", "Choose a directory").ToString(), InitialDirectory, OutDirectory))
 			{
-				if (FPaths::IsUnderDirectory(OutDirectory, FPaths::RootDir()))
+				if (FPaths::IsUnderDirectory(OutDirectory, FPaths::RootDir()) && InMakeRelative)
 				{
 					FPaths::MakePathRelativeTo(OutDirectory, *FPaths::RootDir());
 				}
@@ -592,6 +597,8 @@ namespace ProjectLauncher
 				.Text_Lambda( GetString )
 				.OnTextCommitted_Lambda( SetString )
 				.IsEnabled_Lambda( IsEnabled )
+				.OverflowPolicy(ETextOverflowPolicy::Ellipsis)
+				.ToolTipText_Lambda(GetTooltipText)
 			]
 
 			// browse button
@@ -707,6 +714,7 @@ namespace ProjectLauncher
 				.Text_Lambda( GetString )
 				.OnTextCommitted_Lambda( SetString )
 				.IsEnabled_Lambda( IsEnabled )
+				.OverflowPolicy(ETextOverflowPolicy::Ellipsis)
 			]
 
 			// browse button
diff --git a/Engine/Plugins/Developer/ProjectLauncher/Source/ProjectLauncher/Public/ProfileTree/LaunchProfileTreeData.h b/Engine/Plugins/Developer/ProjectLauncher/Source/ProjectLauncher/Public/ProfileTree/LaunchProfileTreeData.h
--- a/Engine/Plugins/Developer/ProjectLauncher/Source/ProjectLauncher/Public/ProfileTree/LaunchProfileTreeData.h
+++ b/Engine/Plugins/Developer/ProjectLauncher/Source/ProjectLauncher/Public/ProfileTree/LaunchProfileTreeData.h
@@ -122,7 +122,7 @@ namespace ProjectLauncher
 			FValidation Validation;
 		};
 		UE_API FLaunchProfileTreeNode& AddString( FText InName, FStringCallbacks&& StringCallbacks, FText ToolTipText = FText::FromString(""), FText HintText = FText::FromString(""));
-		UE_API FLaunchProfileTreeNode& AddDirectoryString(FText InName, FStringCallbacks&& StringCallbacks);
+		UE_API FLaunchProfileTreeNode& AddDirectoryString(FText InName, FStringCallbacks&& StringCallbacks, bool InMakeRelative = true, FText ToolTipText = FText::FromString(""));
 		UE_API FLaunchProfileTreeNode& AddFileString(FText InName, FStringCallbacks&& StringCallbacks, FString FileFilter = TEXT("All files (*.*)|*.*"));
 		UE_API FLaunchProfileTreeNode& AddCommandLineString( FText InName, FStringCallbacks&& StringCallbacks, const TCHAR* InType = nullptr ); // defaults to "Launch" cmdline extensions
 
