Project Launcher now allows you to select maps that are in plugins.

Also fixed an issue where changing certain properties, e.g.

#Synced-CL 52675748

diff --git a/Engine/Plugins/Developer/ProjectLauncher/Source/CommonLaunchExtensions/Private/AdvancedCook/AdvancedCookLaunchExtension.cpp b/Engine/Plugins/Developer/ProjectLauncher/Source/CommonLaunchExtensions/Private/AdvancedCook/AdvancedCookLaunchExtension.cpp
--- a/Engine/Plugins/Developer/ProjectLauncher/Source/CommonLaunchExtensions/Private/AdvancedCook/AdvancedCookLaunchExtension.cpp
+++ b/Engine/Plugins/Developer/ProjectLauncher/Source/CommonLaunchExtensions/Private/AdvancedCook/AdvancedCookLaunchExtension.cpp
@@ -51,6 +51,8 @@ FAdvancedCookLaunchExtensionInstance::FBuildCookRunInstance::FBuildCookRunInstan
 
 	CachedCulturesToCook = GetBuildCookRun()->GetCookedCultures();
 	CultureOption = (CachedCulturesToCook.Num() == 1 && CachedCulturesToCook.Contains(TEXT("en"))) ? ECultureOption::Default : ECultureOption::Custom;
+
+	LoadMapFilterState();
 }
 
 
@@ -105,6 +107,8 @@ TSharedRef<SWidget> FAdvancedCookLaunchExtensionInstance::FBuildCookRunInstance:
 		.OnSelectionChanged(this, &FAdvancedCookLaunchExtensionInstance::FBuildCookRunInstance::SetMapsToCook)
 		.SelectedMaps(this, &FAdvancedCookLaunchExtensionInstance::FBuildCookRunInstance::GetMapsToCook)
 		.ProjectPath_Lambda( [this]() { return GetModel()->GetProjectPath(GetProfile()); })
+		.OnMapFilterChanged(this, &FAdvancedCookLaunchExtensionInstance::FBuildCookRunInstance::OnMapFilterChanged)
+		.MapSourceFlags(MapSourceFlags)
 	;
 
 	return SNew(SVerticalBox)
@@ -268,7 +272,7 @@ void FAdvancedCookLaunchExtensionInstance::FBuildCookRunInstance::SetMapsToCook(
 		GetBuildCookRun()->AddCookedMap(Map);
 	}
 
-	OnPropertyChanged();
+	GetExtensionInstance()->BroadcastPropertyChanged();
 }
 
 TArray<FString> FAdvancedCookLaunchExtensionInstance::FBuildCookRunInstance::GetMapsToCook() const
@@ -317,7 +321,7 @@ void FAdvancedCookLaunchExtensionInstance::FBuildCookRunInstance::SetMapOption(
 		SetMapsToCook(TArray<FString>());
 	}
 
-	OnPropertyChanged();
+	GetExtensionInstance()->BroadcastPropertyChanged();
 
 	MapListView->RefreshMapList();
 }
@@ -333,7 +337,7 @@ void FAdvancedCookLaunchExtensionInstance::FBuildCookRunInstance::SetCulturesToC
 		GetBuildCookRun()->AddCookedCulture(Culture);
 	}
 
-	OnPropertyChanged();
+	GetExtensionInstance()->BroadcastPropertyChanged();
 }
 
 TArray<FString> FAdvancedCookLaunchExtensionInstance::FBuildCookRunInstance::GetCulturesToCook() const
@@ -379,7 +383,7 @@ void FAdvancedCookLaunchExtensionInstance::FBuildCookRunInstance::SetCultureOpti
 		SetCulturesToCook(TArray<FString>{TEXT("en")});
 	}
 
-	OnPropertyChanged();
+	GetExtensionInstance()->BroadcastPropertyChanged();
 
 	CultureListView->RefreshCultureList();
 }
@@ -393,7 +397,54 @@ void FAdvancedCookLaunchExtensionInstance::FBuildCookRunInstance::OnProjectChang
 	}
 }
 
+struct FMapFilterEntry 
+{
+	EMapSourceFlags Flag = EMapSourceFlags::None;
+	const TCHAR* Key = nullptr;
+};
+static const FMapFilterEntry MapFilterEntries[] =
+{
+	{ EMapSourceFlags::MyDeveloperContent,    TEXT("MapFilter_MyDeveloperContent") },
+	{ EMapSourceFlags::OtherDeveloperContent, TEXT("MapFilter_OtherDeveloperContent") },
+	{ EMapSourceFlags::ProjectPlugins,        TEXT("MapFilter_ProjectPluginMaps") },
+	{ EMapSourceFlags::EngineContent,         TEXT("MapFilter_EngineContent") },
+	{ EMapSourceFlags::EnginePlugins,         TEXT("MapFilter_EnginePluginMaps") },
+};
+
+void FAdvancedCookLaunchExtensionInstance::FBuildCookRunInstance::OnMapFilterChanged(EMapSourceFlags NewFlags)
+{
+	MapSourceFlags = NewFlags;
+	SaveMapFilterState();
+	GetExtensionInstance()->BroadcastPropertyChanged();
+}
+
+void FAdvancedCookLaunchExtensionInstance::FBuildCookRunInstance::LoadMapFilterState()
+{
+	const TMap<FString, bool>& CustomBools = GetProfile()->GetCustomBoolProperties();
+	for (const FMapFilterEntry& Entry : MapFilterEntries)
+	{
+		if (const bool* ValuePtr = CustomBools.Find(Entry.Key))
+		{
+			if (*ValuePtr)
+			{
+				MapSourceFlags |= Entry.Flag;
+			}
+			else
+			{
+				MapSourceFlags &= ~Entry.Flag;
+			}
+		}
+	}
+}
 
+void FAdvancedCookLaunchExtensionInstance::FBuildCookRunInstance::SaveMapFilterState()
+{
+	TMap<FString, bool>& CustomBools = GetProfile()->GetCustomBoolProperties();
+	for (const FMapFilterEntry& Entry : MapFilterEntries)
+	{
+		CustomBools.Add(Entry.Key, EnumHasAnyFlags(MapSourceFlags, Entry.Flag));
+	}
+}
 
 TSharedPtr<ProjectLauncher::FLaunchExtensionInstance> FAdvancedCookLaunchExtension::CreateInstanceForProfile( ProjectLauncher::FLaunchExtensionInstance::FArgs& InArgs )
 {
diff --git a/Engine/Plugins/Developer/ProjectLauncher/Source/CommonLaunchExtensions/Private/AdvancedCook/AdvancedCookLaunchExtension.h b/Engine/Plugins/Developer/ProjectLauncher/Source/CommonLaunchExtensions/Private/AdvancedCook/AdvancedCookLaunchExtension.h
--- a/Engine/Plugins/Developer/ProjectLauncher/Source/CommonLaunchExtensions/Private/AdvancedCook/AdvancedCookLaunchExtension.h
+++ b/Engine/Plugins/Developer/ProjectLauncher/Source/CommonLaunchExtensions/Private/AdvancedCook/AdvancedCookLaunchExtension.h
@@ -3,6 +3,7 @@
 #pragma once
 
 #include "Extension/BuildCookRunCommandExtension.h"
+#include "Widgets/Shared/SCustomLaunchMapListView.h"
 
 
 class SCustomLaunchMapListView;
@@ -50,6 +51,12 @@ protected:
 		float MapListHeight = 300.0f;
 		TSharedPtr<SCustomLaunchMapListView> MapListView;
 
+		// Map filter state
+		EMapSourceFlags MapSourceFlags = EMapSourceFlags::None;
+		void OnMapFilterChanged(EMapSourceFlags NewFlags);
+		void LoadMapFilterState();
+		void SaveMapFilterState();
+
 
 		TArray<FString> CachedCulturesToCook;
 		enum class ECultureOption : uint8
diff --git a/Engine/Plugins/Developer/ProjectLauncher/Source/ProjectLauncher/Private/Model/ProjectLauncherModel.cpp b/Engine/Plugins/Developer/ProjectLauncher/Source/ProjectLauncher/Private/Model/ProjectLauncherModel.cpp
--- a/Engine/Plugins/Developer/ProjectLauncher/Source/ProjectLauncher/Private/Model/ProjectLauncherModel.cpp
+++ b/Engine/Plugins/Developer/ProjectLauncher/Source/ProjectLauncher/Private/Model/ProjectLauncherModel.cpp
@@ -1242,32 +1242,19 @@ PRAGMA_ENABLE_DEPRECATION_WARNINGS
 		return (Profile->GetBuildCookRunCommands().Num() == 1 && Profile->GetAutomatedTests().Num() == 0 && Profile->GetUATCommands().Num() == 1);
 	}
 
-	TArray<FString> FModel::GetAndCacheMapPaths(const FString& InOptionalProjectPath, bool bIncludeNonContentDirMaps) // @todo: map list parsing should ideally be asyncronous, showing a spinner in the map selector controls until its finished etc.
+	TArray<FString> FModel::GetAndCacheMapPaths(const FString& InContentDir, bool bIncludeNonContentDirMaps) // @todo: map list parsing should ideally be asyncronous, showing a spinner in the map selector controls until its finished etc.
 	{
-		// prepare values
-		FString ProjectPath;
-		FString ProjectName;
-		FString ContentDir;
-		bool bWantEngineMaps = InOptionalProjectPath.IsEmpty();
-		if (bWantEngineMaps)
-		{
-			ContentDir = FPaths::Combine(FPaths::EngineContentDir(), TEXT("Maps"));
-		}
-		else
-		{
-			ProjectPath = FPaths::IsRelative(InOptionalProjectPath) ? FPaths::Combine(FPaths::RootDir(), InOptionalProjectPath) : InOptionalProjectPath;
-			ProjectName = FPaths::GetBaseFilename(ProjectPath);
-			ContentDir = FPaths::Combine(ProjectPath, TEXT("Content")); // @todo: do we want to show plugin maps as well?
-		}
-				
-		// the editor has access to the asset registry which is must faster and will be up to date if new maps are added at runtime - if this is the project that's selected in the editor
+		// the editor has access to the asset registry which is much faster and will be up to date if new maps are added at runtime
 #if WITH_EDITOR
 		if (FPaths::IsProjectFilePathSet())
 		{
-			FString GlobalProjectFilePath = FPaths::ConvertRelativePathToFull(FPaths::GetProjectFilePath());
-			FString GlobalProjectPath = FPaths::GetPath(GlobalProjectFilePath);
+			FString GlobalProjectPath = FPaths::ConvertRelativePathToFull(FPaths::GetPath(FPaths::GetProjectFilePath()));
+			FString EngineDir = FPaths::ConvertRelativePathToFull(FPaths::EngineDir());
+			FString ContentDirFull = FPaths::ConvertRelativePathToFull(InContentDir);
 
-			if (!GlobalProjectPath.IsEmpty() && (GlobalProjectPath == ProjectPath || bWantEngineMaps) )
+			// only use the asset registry when the content directory is under the current project or engine tree
+			// (the asset registry only indexes content from the currently open project, engine, and their plugins)
+			if (!GlobalProjectPath.IsEmpty() && (FPaths::IsUnderDirectory(ContentDirFull, GlobalProjectPath) || FPaths::IsUnderDirectory(ContentDirFull, EngineDir) || bIncludeNonContentDirMaps))
 			{
 				// gather all world asset metadata
 				FAssetRegistryModule& AssetRegistryModule = FModuleManager::LoadModuleChecked<FAssetRegistryModule>("AssetRegistry");
@@ -1281,7 +1268,7 @@ PRAGMA_ENABLE_DEPRECATION_WARNINGS
 					FString MapFileName = FPackageName::LongPackageNameToFilename(MapAsset.PackageName.ToString(), FPackageName::GetMapPackageExtension());
 					MapFileName = FPaths::ConvertRelativePathToFull(MapFileName);
 
-					if (FPaths::IsUnderDirectory(MapFileName, ContentDir) || bIncludeNonContentDirMaps)
+					if (FPaths::IsUnderDirectory(MapFileName, ContentDirFull) || bIncludeNonContentDirMaps)
 					{
 						MapFileList.Add( MapFileName );
 					}
@@ -1293,33 +1280,31 @@ PRAGMA_ENABLE_DEPRECATION_WARNINGS
 		}
 #endif
 
-		// check to see if we've cached this project's maps already
-		const TArray<FString>* CachedProjectMapListPtr = CachedMapPaths.Find(ProjectName);
-		if (CachedProjectMapListPtr != nullptr)
+		// check to see if we've cached these maps already
+		const TArray<FString>* CachedMapListPtr = CachedMapPaths.Find(InContentDir);
+		if (CachedMapListPtr != nullptr)
 		{
-			return *CachedProjectMapListPtr;
+			return *CachedMapListPtr;
 		}
-		else
-		{
-			// this is slow & currently blocking so show a wait dialog. 
-			// annoyingly this doesn't show up in UnrealFrontend because that uses FFeedbackContext rather than FFeedbackContextEditor
-			FScopedSlowTask SlowTask(0, LOCTEXT("CacheProjectMapsDesc","Caching project maps"));
-			SlowTask.MakeDialog();
 
-			// search for map files
-			TArray<FString> MapFileList;
-			const FString WildCard = FString::Printf(TEXT("*%s"), *FPackageName::GetMapPackageExtension());
-			IFileManager::Get().FindFilesRecursive(MapFileList, *ContentDir, *WildCard, true, false);
+		// this is slow & currently blocking so show a wait dialog. 
+		// annoyingly this doesn't show up in UnrealFrontend because that uses FFeedbackContext rather than FFeedbackContextEditor
+		FScopedSlowTask SlowTask(0, LOCTEXT("CacheProjectMapsDesc","Caching project maps"));
+		SlowTask.MakeDialog();
 
-			MapFileList.Sort();
-			CachedMapPaths.Add(ProjectName, MapFileList);
-			return MapFileList;
-		}
+		// search for map files
+		TArray<FString> MapFileList;
+		const FString WildCard = FString::Printf(TEXT("*%s"), *FPackageName::GetMapPackageExtension());
+		IFileManager::Get().FindFilesRecursive(MapFileList, *InContentDir, *WildCard, true, false);
+
+		MapFileList.Sort();
+		CachedMapPaths.Add(InContentDir, MapFileList);
+		return MapFileList;
 	}
 
 	TArray<FString> FModel::GetAvailableProjectMapNames( const FString& InProjectPath, bool bIncludeNonContentDirMaps)
 	{
-		TArray<FString> Maps = GetAndCacheMapPaths(InProjectPath, bIncludeNonContentDirMaps);
+		TArray<FString> Maps = GetAvailableProjectMapPaths(InProjectPath, bIncludeNonContentDirMaps);
 		for (FString& Map : Maps)
 		{
 			Map = FPaths::GetBaseFilename(Map);
@@ -1329,12 +1314,14 @@ PRAGMA_ENABLE_DEPRECATION_WARNINGS
 
 	TArray<FString> FModel::GetAvailableProjectMapPaths( const FString& InProjectPath, bool bIncludeNonContentDirMaps)
 	{
-		return GetAndCacheMapPaths(InProjectPath, bIncludeNonContentDirMaps);
+		const FString ProjectPath = FPaths::IsRelative(InProjectPath) ? FPaths::Combine(FPaths::RootDir(), InProjectPath) : InProjectPath;
+		const FString ContentDir = FPaths::Combine(ProjectPath, TEXT("Content"));
+		return GetAndCacheMapPaths(ContentDir, bIncludeNonContentDirMaps);
 	}
 
 	TArray<FString> FModel::GetAvailableEngineMapNames()
 	{
-		TArray<FString> Maps = GetAndCacheMapPaths(FString());
+		TArray<FString> Maps = GetAvailableEngineMapPaths();
 		for (FString& Map : Maps)
 		{
 			Map = FPaths::GetBaseFilename(Map);
@@ -1345,7 +1332,22 @@ PRAGMA_ENABLE_DEPRECATION_WARNINGS
 
 	TArray<FString> FModel::GetAvailableEngineMapPaths()
 	{
-		return GetAndCacheMapPaths(FString());
+		return GetAndCacheMapPaths(FPaths::Combine(FPaths::EngineContentDir(), TEXT("Maps")));
+	}
+
+	TArray<FString> FModel::GetAvailablePluginMapNames(const FString& PluginContentDir)
+	{
+		TArray<FString> Maps = GetAvailablePluginMapPaths(PluginContentDir);
+		for (FString& Map : Maps)
+		{
+			Map = FPaths::GetBaseFilename(Map);
+		}
+		return MoveTemp(Maps);
+	}
+
+	TArray<FString> FModel::GetAvailablePluginMapPaths(const FString& PluginContentDir)
+	{
+		return GetAndCacheMapPaths(PluginContentDir);
 	}
 
 	void FModel::GetUATCommandsDetail( const ILauncherProfileRef& Profile, FUATCommandsDetail& OutUATCommandsDetail)
diff --git a/Engine/Plugins/Developer/ProjectLauncher/Source/ProjectLauncher/Private/ProfileTree/GenericProfileTreeBuilder.cpp b/Engine/Plugins/Developer/ProjectLauncher/Source/ProjectLauncher/Private/ProfileTree/GenericProfileTreeBuilder.cpp
--- a/Engine/Plugins/Developer/ProjectLauncher/Source/ProjectLauncher/Private/ProfileTree/GenericProfileTreeBuilder.cpp
+++ b/Engine/Plugins/Developer/ProjectLauncher/Source/ProjectLauncher/Private/ProfileTree/GenericProfileTreeBuilder.cpp
@@ -904,6 +904,8 @@ namespace ProjectLauncher
 		}
 
 		TreeData->OnPropertyChanged();
+
+		bStartupMapCacheDirty = true;
 	}
 
 
@@ -917,11 +919,36 @@ namespace ProjectLauncher
 
 		CachedStartupMaps.Reset();
 		CachedStartupMaps.Add(MakeShared<FString>());
-		for ( const FString& Map : Model->GetAvailableProjectMapNames(Profile->GetProjectBasePath()))
+
+		TSet<FString> AddedMaps;
+		for (const FString& Map : Model->GetAvailableProjectMapNames(Profile->GetProjectBasePath()))
 		{
 			CachedStartupMaps.Add(MakeShared<FString>(Map));
+			AddedMaps.Add(Map);
 		}
 
+		// Append any maps selected for cooking that aren't project maps
+		// (e.g. plugin maps), so they can be set as the Initial Map
+		for (const ILauncherProfileBuildCookRunRef& BuildCookRun : Profile->GetBuildCookRunCommands())
+		{
+			for (const FString& CookedMap : BuildCookRun->GetCookedMaps())
+			{
+				if (!AddedMaps.Contains(CookedMap))
+				{
+					CachedStartupMaps.Add(MakeShared<FString>(CookedMap));
+					AddedMaps.Add(CookedMap);
+				}
+			}
+		}
+
+		// Refresh the combo box's internal filtered list to match the new options
+		for (const TPair<ILauncherProfileBuildCookRunRef,TSharedRef<FBuildCookRun>>& Pair : BuildCookRunData)
+		{
+			if (Pair.Value->InitalMapCombo.IsValid())
+			{
+				Pair.Value->InitalMapCombo->RefreshOptions();
+			}
+		}
 	}
 
 	void FGenericProfileTreeBuilder::FBuildCookRun::CacheArchitectures()
@@ -1070,8 +1097,6 @@ namespace ProjectLauncher
 		}
 
 		OnPropertyChanged();
-
-		bStartupMapCacheDirty = true;
 	}
 
 	bool FGenericProfileTreeBuilder::HasProject() const
diff --git a/Engine/Plugins/Developer/ProjectLauncher/Source/ProjectLauncher/Private/Widgets/Shared/SCustomLaunchMapListView.cpp b/Engine/Plugins/Developer/ProjectLauncher/Source/ProjectLauncher/Private/Widgets/Shared/SCustomLaunchMapListView.cpp
--- a/Engine/Plugins/Developer/ProjectLauncher/Source/ProjectLauncher/Private/Widgets/Shared/SCustomLaunchMapListView.cpp
+++ b/Engine/Plugins/Developer/ProjectLauncher/Source/ProjectLauncher/Private/Widgets/Shared/SCustomLaunchMapListView.cpp
@@ -6,6 +6,7 @@
 #include "Widgets/Views/STreeView.h"
 #include "Widgets/Input/SCheckBox.h"
 #include "Widgets/Input/SComboBox.h"
+#include "Widgets/Input/SComboButton.h"
 #include "Widgets/Input/SSearchBox.h"
 #include "Widgets/Text/STextBlock.h"
 #include "Widgets/Images/SImage.h"
@@ -14,7 +15,7 @@
 #include "Styling/AppStyle.h"
 #include "Styling/ProjectLauncherStyle.h"
 #include "Misc/PackageName.h"
-#include "HAL/FileManager.h"
+#include "Interfaces/IPluginManager.h"
 
 
 #define LOCTEXT_NAMESPACE "SCustomLaunchMapListView"
@@ -27,6 +28,8 @@ void SCustomLaunchMapListView::Construct(const FArguments& InArgs, TSharedRef<Pr
 	OnSelectionChanged = InArgs._OnSelectionChanged;
 	SelectedMaps = InArgs._SelectedMaps;
 	ProjectPath = InArgs._ProjectPath;
+	OnMapFilterChanged = InArgs._OnMapFilterChanged;
+	MapSourceFlags = InArgs._MapSourceFlags;
 
 	ChildSlot
 	[
@@ -72,6 +75,32 @@ TSharedRef<SWidget> SCustomLaunchMapListView::MakeControlsWidget()
 			]
 		]
 
+		+SHorizontalBox::Slot()
+		.AutoWidth()
+		.Padding(8, 0, 0, 0)
+		[
+			SNew(SComboButton)
+			.HasDownArrow(true)
+			.ToolTipText(LOCTEXT("FilterMenuTooltip", "Toggle visibility of additional map sources"))
+			.ButtonStyle(&FAppStyle::Get().GetWidgetStyle<FButtonStyle>("SimpleButton"))
+			.OnGetMenuContent_Lambda([this]()
+			{
+				return MakeFilterMenuContent();
+			})
+			.ButtonContent()
+			[
+				SNew(SImage)
+				.DesiredSizeOverride(FVector2D(16,16))
+				.Image(FAppStyle::GetBrush("Icons.Filter"))
+				.ColorAndOpacity_Lambda([this]()
+				{
+					return HasAnyFilterEnabled()
+						? FAppStyle::Get().GetSlateColor("Colors.AccentBlue")
+						: FSlateColor::UseForeground();
+				})
+			]
+		]
+
 		+SHorizontalBox::Slot()
 		.FillWidth(1)
 		.Padding(4,0)
@@ -84,6 +113,62 @@ TSharedRef<SWidget> SCustomLaunchMapListView::MakeControlsWidget()
 }
 END_SLATE_FUNCTION_BUILD_OPTIMIZATION
 
+bool SCustomLaunchMapListView::HasAnyFilterEnabled() const
+{
+	return MapSourceFlags != EMapSourceFlags::None;
+}
+
+TSharedRef<SWidget> SCustomLaunchMapListView::MakeFilterMenuContent()
+{
+	TFunction<TSharedRef<SWidget>(const FText&, EMapSourceFlags)> MakeFilterCheckBox = [this](const FText& Label, EMapSourceFlags Flag) -> TSharedRef<SWidget>
+	{
+		return SNew(SCheckBox)
+			.IsChecked_Lambda([this, Flag]() { return EnumHasAnyFlags(MapSourceFlags, Flag) ? ECheckBoxState::Checked : ECheckBoxState::Unchecked; })
+			.OnCheckStateChanged_Lambda([this, Flag](ECheckBoxState NewState)
+			{
+				if (NewState == ECheckBoxState::Checked)
+				{
+					MapSourceFlags |= Flag;
+				}
+				else
+				{
+					MapSourceFlags &= ~Flag;
+				}
+				RefreshMapList();
+				OnMapFilterChanged.ExecuteIfBound(MapSourceFlags);
+			})
+			[
+				SNew(STextBlock)
+				.Text(Label)
+			];
+	};
+
+	return SNew(SBox)
+		.Padding(FMargin(0, 4))
+		[
+			SNew(SVerticalBox)
+			+ SVerticalBox::Slot().AutoHeight().Padding(8, 4)
+			[
+				MakeFilterCheckBox(LOCTEXT("FilterMyDevContent", "My Developer Content"), EMapSourceFlags::MyDeveloperContent)
+			]
+			+ SVerticalBox::Slot().AutoHeight().Padding(8, 4)
+			[
+				MakeFilterCheckBox(LOCTEXT("FilterOtherDevContent", "Other Developers Content"), EMapSourceFlags::OtherDeveloperContent)
+			]
+			+ SVerticalBox::Slot().AutoHeight().Padding(8, 4)
+			[
+				MakeFilterCheckBox(LOCTEXT("FilterProjectPlugins", "Project Plugins"), EMapSourceFlags::ProjectPlugins)
+			]
+			+ SVerticalBox::Slot().AutoHeight().Padding(8, 4)
+			[
+				MakeFilterCheckBox(LOCTEXT("FilterEngineContent", "Engine Content"), EMapSourceFlags::EngineContent)
+			]
+			+ SVerticalBox::Slot().AutoHeight().Padding(8, 4)
+			[
+				MakeFilterCheckBox(LOCTEXT("FilterEnginePlugins", "Engine Plugins"), EMapSourceFlags::EnginePlugins)
+			]
+		];
+}
 
 void SCustomLaunchMapListView::GetMapTreeNodeChildren(FMapTreeNodePtr Node, TArray<FMapTreeNodePtr>& OutChildren)
 {
@@ -256,7 +341,6 @@ void SCustomLaunchMapListView::RefreshMapListInternal()
 
 	// decide what to show
 	bool bShowProjectMaps = !ProjectPath.Get().IsEmpty();
-	bool bShowEngineMaps = false;
 
 	TArray<FString> AvailableMaps;
 	const FString WildCard = FString::Printf(TEXT("*%s"), *FPackageName::GetMapPackageExtension());
@@ -273,9 +357,34 @@ void SCustomLaunchMapListView::RefreshMapListInternal()
 			ProjectContentDir = FPaths::Combine(FPaths::RootDir(), ProjectContentDir);
 		}
 
+		// determine the developers content directory and current user's folder for filtering
+		const FString DevelopersSubDir = FString(TEXT("/")) + FString(FPaths::DevelopersFolderName()) + TEXT("/");
+		const FString MyDeveloperSubDir = FString(TEXT("/")) + FString(FPaths::DevelopersFolderName()) + TEXT("/") + FPaths::GameUserDeveloperFolderName() + TEXT("/");
+
 		TArray<FString> ProjectMaps = Model->GetAvailableProjectMapPaths(ProjectBasePath);
 		for (FString& ProjectMap : ProjectMaps)
 		{
+			// check if this map is in a developers folder
+			FString RelativePath = ProjectMap;
+			RelativePath.RemoveFromStart(*ProjectContentDir);
+			RelativePath.ReplaceCharInline('\\', '/');
+
+			const bool bIsInDevelopersDir = RelativePath.StartsWith(*DevelopersSubDir);
+			if (bIsInDevelopersDir)
+			{
+				const bool bIsMyDeveloperContent = RelativePath.StartsWith(*MyDeveloperSubDir);
+
+				if (bIsMyDeveloperContent && !EnumHasAnyFlags(MapSourceFlags, EMapSourceFlags::MyDeveloperContent))
+				{
+					continue;
+				}
+
+				if (!bIsMyDeveloperContent && !EnumHasAnyFlags(MapSourceFlags, EMapSourceFlags::OtherDeveloperContent))
+				{
+					continue;
+				}
+			}
+
 			AvailableMaps.AddUnique(FPaths::GetBaseFilename(ProjectMap));
 
 			ProjectMap.RemoveFromStart(*ProjectContentDir);
@@ -283,27 +392,60 @@ void SCustomLaunchMapListView::RefreshMapListInternal()
 			AddMapPath(ProjectName / ProjectMap);
 		}
 	}
-	else
-	{
-		bShowEngineMaps = true;
-	}
-	
-	
-	// show maps from the engine (fallback)
-	if (bShowEngineMaps)
+
+
+	// show maps from the engine
+	if (EnumHasAnyFlags(MapSourceFlags, EMapSourceFlags::EngineContent) || !bShowProjectMaps)
 	{
-		FString EngineMapDir = FPaths::Combine(FPaths::EngineContentDir(), TEXT("Maps"));
+		FString EngineContentDir = FPaths::ConvertRelativePathToFull(FPaths::EngineContentDir());
 		TArray<FString> EngineMaps = Model->GetAvailableEngineMapPaths();
 		for (FString& EngineMap : EngineMaps)
 		{
 			AvailableMaps.AddUnique(FPaths::GetBaseFilename(EngineMap));
 
-			EngineMap.RemoveFromStart(*EngineMapDir);
+			EngineMap = FPaths::ConvertRelativePathToFull(EngineMap);
+			EngineMap.RemoveFromStart(*EngineContentDir);
 			EngineMap.ReplaceCharInline('\\', '/');
 			AddMapPath(TEXT("Engine") / EngineMap);
 		}
 	}
 
+	// show maps from plugins
+	if (EnumHasAnyFlags(MapSourceFlags, EMapSourceFlags::ProjectPlugins | EMapSourceFlags::EnginePlugins))
+	{
+		TArray<TSharedRef<IPlugin>> Plugins = IPluginManager::Get().GetEnabledPluginsWithContent();
+		for (const TSharedRef<IPlugin>& Plugin : Plugins)
+		{
+			FString ContentDir = FPaths::ConvertRelativePathToFull(Plugin->GetContentDir());
+			if (ContentDir.IsEmpty())
+			{
+				continue;
+			}
+
+			EPluginType PluginType = Plugin->GetType();
+			const bool bIsProjectPlugin = (PluginType == EPluginType::Project || PluginType == EPluginType::Mod || PluginType == EPluginType::External);
+			const bool bIsEnginePlugin = (PluginType == EPluginType::Engine || PluginType == EPluginType::Enterprise);
+
+			if (!(bIsProjectPlugin && EnumHasAnyFlags(MapSourceFlags, EMapSourceFlags::ProjectPlugins)) &&
+				!(bIsEnginePlugin && EnumHasAnyFlags(MapSourceFlags, EMapSourceFlags::EnginePlugins)))
+			{
+				continue;
+			}
+
+			FString PluginName = Plugin->GetFriendlyName();
+			TArray<FString> PluginMaps = Model->GetAvailablePluginMapPaths(ContentDir);
+			for (const FString& MapPath : PluginMaps)
+			{
+				AvailableMaps.AddUnique(FPaths::GetBaseFilename(MapPath));
+
+				FString RelativeMap = FPaths::ConvertRelativePathToFull(MapPath);
+				RelativeMap.RemoveFromStart(*ContentDir);
+				RelativeMap.ReplaceCharInline('\\', '/');
+				AddMapPath(PluginName / RelativeMap);
+			}
+		}
+	}
+
 	// add any maps that are selected but we have not found them
 	for (const FString& CookedMap : SelectedMaps.Get() )
 	{
diff --git a/Engine/Plugins/Developer/ProjectLauncher/Source/ProjectLauncher/Public/Model/ProjectLauncherModel.h b/Engine/Plugins/Developer/ProjectLauncher/Source/ProjectLauncher/Public/Model/ProjectLauncherModel.h
--- a/Engine/Plugins/Developer/ProjectLauncher/Source/ProjectLauncher/Public/Model/ProjectLauncherModel.h
+++ b/Engine/Plugins/Developer/ProjectLauncher/Source/ProjectLauncher/Public/Model/ProjectLauncherModel.h
@@ -314,11 +314,13 @@ namespace ProjectLauncher
 		/* Returns true if we can use a simplified layout for this profile */
 		UE_API bool CanUseSimplifiedLayout(const ILauncherProfileRef& Profile);
 
-		/* Get available maps for a project or engine */
+		/* Get available maps for a project, engine, or plugin */
 		UE_API TArray<FString> GetAvailableProjectMapNames(const FString& InProjectPath, bool bIncludeNonContentDirMaps = false);
 		UE_API TArray<FString> GetAvailableProjectMapPaths(const FString& InProjectPath, bool bIncludeNonContentDirMaps = false);
 		UE_API TArray<FString> GetAvailableEngineMapNames();
 		UE_API TArray<FString> GetAvailableEngineMapPaths();
+		UE_API TArray<FString> GetAvailablePluginMapNames(const FString& PluginContentDir);
+		UE_API TArray<FString> GetAvailablePluginMapPaths(const FString& PluginContentDir);
 
 		/* Get the current UAT command information for the given profile */
 		struct FUATCommandsDetail
@@ -355,8 +357,8 @@ namespace ProjectLauncher
 		UE_API void OnModifyLaunchCommandLine( const ILauncherProfileRef& InProfile, EBuildTargetType InBuildTargetType, FString& InOutCommandLine );
 		UE_API void OnModifyUATCommandLine( FUATCommandPostProcessParameters& PostProcessParameters );
 
-		/* Get map paths for a project, caching results */
-		UE_API TArray<FString> GetAndCacheMapPaths(const FString& InOptionalProjectPath, bool bIncludeNonContentDirMaps = false);
+		/* Get map paths under a content directory, caching results */
+		UE_API TArray<FString> GetAndCacheMapPaths(const FString& InContentDir, bool bIncludeNonContentDirMaps = false);
 
 	private:
 		FString ConfigFileName;
diff --git a/Engine/Plugins/Developer/ProjectLauncher/Source/ProjectLauncher/Public/Widgets/Shared/SCustomLaunchMapListView.h b/Engine/Plugins/Developer/ProjectLauncher/Source/ProjectLauncher/Public/Widgets/Shared/SCustomLaunchMapListView.h
--- a/Engine/Plugins/Developer/ProjectLauncher/Source/ProjectLauncher/Public/Widgets/Shared/SCustomLaunchMapListView.h
+++ b/Engine/Plugins/Developer/ProjectLauncher/Source/ProjectLauncher/Public/Widgets/Shared/SCustomLaunchMapListView.h
@@ -6,6 +6,17 @@
 #include "Widgets/SCompoundWidget.h"
 #include "Model/ProjectLauncherModel.h"
 
+enum class EMapSourceFlags : uint8
+{
+	None                  = 0,
+	MyDeveloperContent    = 1 << 0,
+	OtherDeveloperContent = 1 << 1,
+	ProjectPlugins        = 1 << 2,
+	EngineContent         = 1 << 3,
+	EnginePlugins         = 1 << 4,
+};
+ENUM_CLASS_FLAGS(EMapSourceFlags);
+
 #define UE_API PROJECTLAUNCHER_API
 
 class ITableRow;
@@ -17,11 +28,14 @@ class SCustomLaunchMapListView
 {
 public:
 	DECLARE_DELEGATE_OneParam(FOnSelectionChanged, TArray<FString> );
+	DECLARE_DELEGATE_OneParam(FOnMapFilterChanged, EMapSourceFlags);
 
 	SLATE_BEGIN_ARGS(SCustomLaunchMapListView){}
 		SLATE_EVENT(FOnSelectionChanged, OnSelectionChanged);
 		SLATE_ATTRIBUTE(TArray<FString>, SelectedMaps);
 		SLATE_ATTRIBUTE(FString, ProjectPath)
+		SLATE_EVENT(FOnMapFilterChanged, OnMapFilterChanged)
+		SLATE_ARGUMENT(EMapSourceFlags, MapSourceFlags)
 	SLATE_END_ARGS()
 
 public:
@@ -35,10 +49,17 @@ protected:
 	TAttribute<TArray<FString>> SelectedMaps;
 	TAttribute<FString> ProjectPath;
 	FOnSelectionChanged OnSelectionChanged;
+	FOnMapFilterChanged OnMapFilterChanged;
 
 private:
 	bool bShowFolders = true; // whether to display the available maps in a hierarchy or flat list
 
+	// Content visibility flags (toggled via filter menu)
+	EMapSourceFlags MapSourceFlags = EMapSourceFlags::None;
+
+	UE_API TSharedRef<SWidget> MakeFilterMenuContent();
+	bool HasAnyFilterEnabled() const;
+
 	struct FMapTreeNode
 	{
 		FString Name;
